ElectraPlayer: added ElectraCDM with ClearKey for content protection

#rnx

diff --git a/Engine/Plugins/Media/ElectraCDM/ElectraCDM.uplugin b/Engine/Plugins/Media/ElectraCDM/ElectraCDM.uplugin
new file mode 100644
--- /dev/null
+++ b/Engine/Plugins/Media/ElectraCDM/ElectraCDM.uplugin
@@ -0,0 +1,26 @@
+{
+	"FileVersion" : 1,
+	"Version" : 0,
+	"VersionName" : "0.1",
+	"FriendlyName" : "Electra Player Content Decryption Module",
+	"Description" : "Content Decryption Module for Electra Player Media Playback",
+	"Category" : "Media Players",
+	"CreatedBy" : "Epic Games, Inc.",
+	"CreatedByURL" : "http://epicgames.com",
+	"DocsURL" : "https://docs.unrealengine.com/en-US/Engine/MediaFramework/Overview",
+	"MarketplaceURL" : "",
+	"SupportURL" : "",
+	"EnabledByDefault" : false,
+	"CanContainContent" : false,
+	"IsBetaVersion" : true,
+	"Installed" : false,
+	"Modules": [
+		{
+			"Name": "ElectraCDM",
+			"Type": "RuntimeNoCommandlet",
+			"LoadingPhase": "PostEngineInit",
+			"WhitelistPlatforms": [ "Win64", "Mac", "IOS", "Android" ],
+			"BlacklistTargets": [ "Server" ]
+		},
+	]
+}
\ No newline at end of file
diff --git a/Engine/Plugins/Media/ElectraCDM/Source/ElectraCDM/ElectraCDM.Build.cs b/Engine/Plugins/Media/ElectraCDM/Source/ElectraCDM/ElectraCDM.Build.cs
new file mode 100644
--- /dev/null
+++ b/Engine/Plugins/Media/ElectraCDM/Source/ElectraCDM/ElectraCDM.Build.cs
@@ -0,0 +1,18 @@
+// Copyright Epic Games, Inc. All Rights Reserved.
+
+namespace UnrealBuildTool.Rules
+{
+	public class ElectraCDM : ModuleRules
+	{
+		public ElectraCDM(ReadOnlyTargetRules Target) : base(Target)
+		{
+			PCHUsage = PCHUsageMode.NoPCHs;     // to get around the "xzy must be first header included" error in 3rdParty lib
+			PublicDependencyModuleNames.AddRange(
+				new string[]
+				{
+					"Core",
+					"Json"
+				});
+		}
+	}
+}
diff --git a/Engine/Plugins/Media/ElectraCDM/Source/ElectraCDM/Private/ClearKey/ClearKeyCDM.cpp b/Engine/Plugins/Media/ElectraCDM/Source/ElectraCDM/Private/ClearKey/ClearKeyCDM.cpp
new file mode 100644
--- /dev/null
+++ b/Engine/Plugins/Media/ElectraCDM/Source/ElectraCDM/Private/ClearKey/ClearKeyCDM.cpp
@@ -0,0 +1,1206 @@
+// Copyright Epic Games, Inc. All Rights Reserved.
+
+#include "ClearKey/ClearKeyCDM.h"
+#include "ElectraCDM.h"
+#include "ElectraCDMClient.h"
+#include "Crypto/StreamCryptoAES128.h"
+#include <Misc/Base64.h>
+#include <Misc/ScopeLock.h>
+#include <Dom/JsonObject.h>
+#include <Serialization/JsonReader.h>
+#include <Serialization/JsonSerializer.h>
+
+
+#define ENABLE_LEGACY_RAWKEY_OVERRIDE 1
+
+namespace ElectraCDM
+{
+
+namespace ClearKeyCDM
+{
+static const TCHAR* const JSONTextPropertyName = TEXT("#text");
+static FString UrlAttrib_dashif_laurl(TEXT("dashif:laurl"));
+static FString UrlAttrib_laurl(TEXT("laurl"));
+static FString UrlAttrib_clearkey_laurl(TEXT("clearkey:Laurl"));
+static FString UrlAttrib_Laurl(TEXT("Laurl"));
+};
+
+struct FClearKeyKIDKey
+{
+	TArray<uint8> KID;
+	TArray<uint8> Key;
+};
+
+
+class FClearKeyCDM : public IClearKeyCDM, public IMediaCDMCapabilities, public TSharedFromThis<FClearKeyCDM, ESPMode::ThreadSafe>
+{
+public:
+	FClearKeyCDM() = default;
+	virtual ~FClearKeyCDM() = default;
+	virtual FString GetLastErrorMessage() override;
+	virtual const FString& GetSchemeID() override;
+	virtual void GetCDMCustomJSONPrefixes(FString& OutAttributePrefix, FString& OutTextPropertyName, bool& bOutNoNamespaces) override;
+	virtual TSharedPtr<IMediaCDMCapabilities, ESPMode::ThreadSafe> GetCDMCapabilities(const FString& InValue, const FString& InAdditionalElements) override;
+	virtual ECDMError CreateDRMClient(TSharedPtr<IMediaCDMClient, ESPMode::ThreadSafe>& OutClient, IMediaCDM::IPlayerSession* InForPlayerSession, const TArray<IMediaCDM::FCDMCandidate>& InCandidates) override;
+	virtual ECDMError ReleasePlayerSessionKeys(IMediaCDM::IPlayerSession* PlayerSession) override;
+
+	virtual ESupportResult SupportsType(const FString& InMimeType) override;
+	virtual ESupportResult RequiresSecureDecoder(const FString& InMimeType) override;
+
+	void AddPlayerSessionKeys(IMediaCDM::IPlayerSession* InPlayerSession, const TArray<FClearKeyKIDKey>& InNewSessionKeys);
+	bool GetPlayerSessionKey(FClearKeyKIDKey& OutKey, IMediaCDM::IPlayerSession* InPlayerSession, const TArray<uint8>& InForKID);
+
+public:
+	static TSharedPtr<FClearKeyCDM, ESPMode::ThreadSafe> Get();
+
+	FCriticalSection Lock;
+	TMap<IMediaCDM::IPlayerSession*, TArray<FClearKeyKIDKey>> ActiveLicensesPerPlayer;
+	FString LastErrorMessage;
+};
+
+
+class FClearKeyDRMDecrypter : public IMediaCDMDecrypter
+{
+public:
+	FClearKeyDRMDecrypter();
+	void Initialize(const FString& InMimeType);
+	void SetLicenseKeys(const TArray<FClearKeyKIDKey>& InLicenseKeys);
+	void SetState(ECDMState InNewState);
+	void SetLastErrorMessage(const FString& InNewErrorMessage);
+	virtual ~FClearKeyDRMDecrypter();
+	virtual ECDMState GetState() override;
+	virtual FString GetLastErrorMessage() override;
+	virtual ECDMError UpdateInitDataFromPSSH(const TArray<uint8>& InPSSHData) override;
+	virtual ECDMError UpdateInitDataFromMultiplePSSH(const TArray<TArray<uint8>>& InPSSHData) override;
+	virtual ECDMError UpdateFromURL(const FString& InURL, const FString& InAdditionalElements) override;
+	virtual bool IsBlockStreamDecrypter() override;
+	virtual void Reinitialize() override;
+	virtual ECDMError DecryptInPlace(uint8* InOutData, int32 InNumDataBytes, const FMediaCDMSampleInfo& InSampleInfo) override;
+	virtual ECDMError BlockStreamDecryptStart(IStreamDecryptHandle*& OutStreamDecryptContext) override;
+	virtual ECDMError BlockStreamDecryptInPlace(IStreamDecryptHandle* InOutStreamDecryptContext, int32& OutNumBytesDecrypted, uint8* InOutData, int32 InNumDataBytes, const FMediaCDMSampleInfo& InSampleInfo, bool bIsLastBlock) override;
+	virtual ECDMError BlockStreamDecryptEnd(IStreamDecryptHandle* InStreamDecryptContext) override;
+
+private:
+	struct FKeyDecrpyter
+	{
+		FClearKeyKIDKey KIDKey;
+		TSharedPtr<ElectraCDM::IStreamDecrypterAES128, ESPMode::ThreadSafe> Decrypter;
+		ECDMState State = ECDMState::Idle;
+	};
+
+	TSharedPtr<ElectraCDM::IStreamDecrypterAES128, ESPMode::ThreadSafe> GetDecrypterForKID(const TArray<uint8>& KID);
+
+	FCriticalSection Lock;
+	TArray<FClearKeyKIDKey> LicenseKeys;
+	TArray<FKeyDecrpyter> KeyDecrypters;
+	ECDMState CurrentState = ECDMState::Idle;
+	FString LastErrorMsg;
+};
+
+
+class FClearKeyDRMClient : public IMediaCDMClient, public TSharedFromThis<FClearKeyDRMClient, ESPMode::ThreadSafe>
+{
+public:
+	virtual ~FClearKeyDRMClient();
+
+	FClearKeyDRMClient();
+	void Initialize(TSharedPtr<FClearKeyCDM, ESPMode::ThreadSafe> InOwningCDM, IMediaCDM::IPlayerSession* InForPlayerSession, const TArray<IMediaCDM::FCDMCandidate>& InCDMConfigurations);
+
+	virtual ECDMState GetState() override;
+	virtual FString GetLastErrorMessage() override;
+
+	virtual void RegisterEventListener(TWeakPtr<IMediaCDMEventListener, ESPMode::ThreadSafe> InEventListener) override;
+	virtual void UnregisterEventListener(TWeakPtr<IMediaCDMEventListener, ESPMode::ThreadSafe> InEventListener) override;
+	virtual void PrepareLicenses() override;
+	virtual void SetLicenseServerURL(const FString& InLicenseServerURL) override;
+	virtual void GetLicenseKeyURL(FString& OutLicenseURL) override;
+	virtual void GetLicenseKeyRequestData(TArray<uint8>& OutKeyRequestData, FString& OutHttpMethod, TArray<FString>& OutHttpHeaders, uint32& OutFlags) override;
+	virtual ECDMError SetLicenseKeyResponseData(void* InEventId, int32 HttpResponseCode, const TArray<uint8>& InKeyResponseData) override;
+	virtual ECDMError CreateDecrypter(TSharedPtr<IMediaCDMDecrypter, ESPMode::ThreadSafe>& OutDecrypter, const FString& InMimeType) override;
+
+private:
+	int32 PrepareKIDsToRequest();
+	void AddKeyKID(const FClearKeyKIDKey& InKeyKid);
+	void AddKeyKIDs(const TArray<FClearKeyKIDKey>& InKeyKids);
+	void FireEvent(IMediaCDMEventListener::ECDMEventType InEvent);
+	void RemoveStaleDecrypters();
+	void UpdateKeyWithDecrypters();
+	void UpdateStateWithDecrypters(ECDMState InNewState);
+	void GetValuesFromConfigurations();
+	bool GetURLsFrom(TArray<FString>& OutURLs, const TSharedPtr<FJsonValue>& JSONValue, const FString& PropertyName, const FString& AltPropertyName);
+	bool GetURLsFrom(TArray<FString>& OutURLs, const TSharedPtr<FJsonObject>& JSON, const FString& PropertyName, const FString& AltPropertyName);
+
+	FCriticalSection Lock;
+	IMediaCDM::IPlayerSession* PlayerSession = nullptr;
+	TWeakPtr<FClearKeyCDM, ESPMode::ThreadSafe> OwningCDM;
+	TArray<IMediaCDM::FCDMCandidate> CDMConfigurations;
+	TArray<TWeakPtr<IMediaCDMEventListener, ESPMode::ThreadSafe>> Listeners;
+	TArray<TWeakPtr<FClearKeyDRMDecrypter, ESPMode::ThreadSafe>> Decrypters;
+	TArray<FString> PendingRequiredKIDs;
+	TArray<FClearKeyKIDKey> LicenseKeys;
+	TOptional<FString> LicenseServerURLOverride;
+	TArray<FString> LicenseServerURLsFromConfigs;
+	ECDMState CurrentState = ECDMState::Idle;
+	FString LastErrorMsg;
+};
+
+/***************************************************************************************************************************************************/
+/***************************************************************************************************************************************************/
+/***************************************************************************************************************************************************/
+namespace
+{
+	/*
+		ClearKey license requests and responses use base64url encoding.
+		See https://www.w3.org/TR/encrypted-media/#clear-key-request-format
+		Section 9.1.3 License Request Format
+	*/
+
+	static FString Base64UrlEncode(const TArray<uint8>& InData)
+	{
+		FString b64 = FBase64::Encode(InData);
+		// Base64Url encoding replaces '+' and '/' with '-' and '_' respectively.
+		b64.ReplaceCharInline(TCHAR('+'), TCHAR('-'), ESearchCase::IgnoreCase);
+		b64.ReplaceCharInline(TCHAR('/'), TCHAR('_'), ESearchCase::IgnoreCase);
+		return b64;
+	}
+	static bool Base64UrlDecode(TArray<uint8>& OutData, FString InString)
+	{
+		InString.ReplaceCharInline(TCHAR('-'), TCHAR('+'), ESearchCase::IgnoreCase);
+		InString.ReplaceCharInline(TCHAR('_'), TCHAR('/'), ESearchCase::IgnoreCase);
+		return FBase64::Decode(InString, OutData);
+	}
+
+	static FString StripDashesFromKID(const FString& InKID)
+	{
+		return InKID.Replace(TEXT("-"), TEXT(""), ESearchCase::CaseSensitive);
+	}
+
+	static void ConvertKIDToBin(TArray<uint8>& OutBinKID, const FString& InKID)
+	{
+		OutBinKID.Empty();
+		check((InKID.Len() % 2) == 0);
+		if ((InKID.Len() % 2) == 0)
+		{
+			OutBinKID.AddUninitialized(InKID.Len() / 2);
+			HexToBytes(InKID, OutBinKID.GetData());
+		}
+	}
+
+	static FString ConvertKIDToBase64(const FString& InKID)
+	{
+		TArray<uint8> BinKID;
+		ConvertKIDToBin(BinKID, InKID);
+		FString b64 = Base64UrlEncode(BinKID);
+		// Chop off trailing padding.
+		return b64.Replace(TEXT("="), TEXT(""), ESearchCase::CaseSensitive);
+	}
+
+	void StringToArray(TArray<uint8>& OutArray, const FString& InString)
+	{
+		FTCHARToUTF8 cnv(*InString);
+		int32 Len = cnv.Length();
+		OutArray.AddUninitialized(Len);
+		FMemory::Memcpy(OutArray.GetData(), cnv.Get(), Len);
+	}
+
+	FString ArrayToString(const TArray<uint8>& InArray)
+	{
+		FUTF8ToTCHAR cnv((const ANSICHAR*)InArray.GetData(), InArray.Num());
+		FString UTF8Text(cnv.Length(), cnv.Get());
+		return MoveTemp(UTF8Text);
+	}
+
+}
+
+
+/***************************************************************************************************************************************************/
+/***************************************************************************************************************************************************/
+/***************************************************************************************************************************************************/
+
+
+//-----------------------------------------------------------------------------
+/**
+ * Registers this ClearKey CDM with the CDM manager
+ */
+void IClearKeyCDM::RegisterWith(IMediaCDM& InDRMManager)
+{
+	InDRMManager.RegisterCDM(FClearKeyCDM::Get());
+}
+
+//-----------------------------------------------------------------------------
+/**
+ * Returns the singleton of this CDM system.
+ */
+TSharedPtr<FClearKeyCDM, ESPMode::ThreadSafe> FClearKeyCDM::Get()
+{
+	static TSharedPtr<FClearKeyCDM, ESPMode::ThreadSafe> This = MakeShared<FClearKeyCDM, ESPMode::ThreadSafe>();
+	return This;
+}
+
+//-----------------------------------------------------------------------------
+/**
+ * Returns the scheme ID of ClearKey.
+ * This is the official ID used with DASH.
+ */
+const FString& FClearKeyCDM::GetSchemeID()
+{
+	static FString SchemeID(TEXT("e2719d58-a985-b3c9-781a-b030af78d30e"));
+	return SchemeID;
+}
+
+//-----------------------------------------------------------------------------
+/**
+ * Returns the most recent error message.
+ */
+FString FClearKeyCDM::GetLastErrorMessage()
+{
+	FScopeLock lock(&Lock);
+	return LastErrorMessage;
+}
+
+//-----------------------------------------------------------------------------
+/**
+ * Returns the expected element prefixes for the AdditionalElements JSON.
+ */
+void FClearKeyCDM::GetCDMCustomJSONPrefixes(FString& OutAttributePrefix, FString& OutTextPropertyName, bool& bOutNoNamespaces)
+{
+	OutAttributePrefix.Empty();
+	OutTextPropertyName = ClearKeyCDM::JSONTextPropertyName;
+	bOutNoNamespaces = false;	// keep the namespaces to differentiate between dashif:laurl and clearkey:Laurl
+}
+
+//-----------------------------------------------------------------------------
+/**
+ * Returns the capability interface of this CDM.
+ */
+TSharedPtr<IMediaCDMCapabilities, ESPMode::ThreadSafe> FClearKeyCDM::GetCDMCapabilities(const FString& InValue, const FString& InAdditionalElements)
+{
+	if (InValue.Equals(TEXT("ClearKey1.0")))
+	{
+		return AsShared();
+	}
+	return nullptr;
+}
+
+//-----------------------------------------------------------------------------
+/**
+ * Creates a client instance of this CDM.
+ * The application may create one or more instances, possibly for different key IDs.
+ * In the context of DASH there is to be one client per Period that contains one or
+ * more AdaptationSets that are encrypted.
+ */
+ECDMError FClearKeyCDM::CreateDRMClient(TSharedPtr<IMediaCDMClient, ESPMode::ThreadSafe>& OutClient, IMediaCDM::IPlayerSession* InForPlayerSession, const TArray<IMediaCDM::FCDMCandidate>& InCandidates)
+{
+	FClearKeyDRMClient* NewClient = new FClearKeyDRMClient;
+	NewClient->Initialize(AsShared(), InForPlayerSession, InCandidates);
+	OutClient = MakeShareable(NewClient);
+	FScopeLock lock(&Lock);
+	LastErrorMessage.Empty();
+	return ECDMError::Success;
+}
+
+//-----------------------------------------------------------------------------
+/**
+ * Adds keys to the specified player session.
+ */
+void FClearKeyCDM::AddPlayerSessionKeys(IMediaCDM::IPlayerSession* InPlayerSession, const TArray<FClearKeyKIDKey>& InNewSessionKeys)
+{
+	FScopeLock lock(&Lock);
+	TArray<FClearKeyKIDKey>& Keys = ActiveLicensesPerPlayer.FindOrAdd(InPlayerSession);
+	for(auto &NewKey : InNewSessionKeys)
+	{
+		bool bHaveAlready = false;
+		for(auto &HaveKey : Keys)
+		{
+			if (HaveKey.KID == NewKey.KID)
+			{
+				bHaveAlready = true;
+				break;
+			}
+		}
+		if (!bHaveAlready)
+		{
+			Keys.Emplace(NewKey);
+		}
+	}
+}
+
+//-----------------------------------------------------------------------------
+/**
+ * Returns a player session's key for the specified KID.
+ */
+bool FClearKeyCDM::GetPlayerSessionKey(FClearKeyKIDKey& OutKey, IMediaCDM::IPlayerSession* InPlayerSession, const TArray<uint8>& InForKID)
+{
+	FScopeLock lock(&Lock);
+	const TArray<FClearKeyKIDKey>* Keys = ActiveLicensesPerPlayer.Find(InPlayerSession);
+	if (Keys)
+	{
+		for(auto &Key : *Keys)
+		{
+			if (Key.KID == InForKID)
+			{
+				OutKey = Key;
+				return true;
+			}
+		}
+	}
+	return false;
+}
+
+//-----------------------------------------------------------------------------
+/**
+ * Releases all keys the specified player session has acquired.
+ */
+ECDMError FClearKeyCDM::ReleasePlayerSessionKeys(IMediaCDM::IPlayerSession* PlayerSession)
+{
+	FScopeLock lock(&Lock);
+	LastErrorMessage.Empty();
+	ActiveLicensesPerPlayer.Remove(PlayerSession);
+	return ECDMError::Success;
+}
+
+
+
+//-----------------------------------------------------------------------------
+/**
+ * Returns if a media stream of a given format can be decrypted with this CDM.
+ * The mime type should include a codecs="..." component and if it is video it
+ * should also have a resolution=...x... component.
+ */
+IMediaCDMCapabilities::ESupportResult FClearKeyCDM::SupportsType(const FString& InMimeType)
+{
+	// Everything is supported.
+	return IMediaCDMCapabilities::ESupportResult::Supported;
+}
+
+//-----------------------------------------------------------------------------
+/**
+ * Returns whether or not for a particular media stream format a secure decoder is
+ * required to be used.
+ */
+IMediaCDMCapabilities::ESupportResult FClearKeyCDM::RequiresSecureDecoder(const FString& InMimeType)
+{
+	// This is ClearKey... there is no real security here to begin with.
+	return IMediaCDMCapabilities::ESupportResult::SecureDecoderNotRequired;
+}
+
+
+/***************************************************************************************************************************************************/
+/***************************************************************************************************************************************************/
+/***************************************************************************************************************************************************/
+
+
+//-----------------------------------------------------------------------------
+/**
+ * Construct a new client
+ */
+FClearKeyDRMClient::FClearKeyDRMClient()
+{
+}
+
+//-----------------------------------------------------------------------------
+/**
+ * Destroy a client
+ */
+FClearKeyDRMClient::~FClearKeyDRMClient()
+{
+}
+
+//-----------------------------------------------------------------------------
+/**
+ * Returns the client's current state.
+ */
+ECDMState FClearKeyDRMClient::GetState()
+{
+	FScopeLock lock(&Lock);
+	return CurrentState;
+}
+
+//-----------------------------------------------------------------------------
+/**
+ * Returns the client's most recent error message.
+ */
+FString FClearKeyDRMClient::GetLastErrorMessage()
+{
+	FScopeLock lock(&Lock);
+	return LastErrorMsg;
+}
+
+//-----------------------------------------------------------------------------
+/**
+ * Initializes the client.
+ */
+void FClearKeyDRMClient::Initialize(TSharedPtr<FClearKeyCDM, ESPMode::ThreadSafe> InOwningCDM, IMediaCDM::IPlayerSession* InForPlayerSession, const TArray<IMediaCDM::FCDMCandidate>& InCDMConfigurations)
+{
+	FScopeLock lock(&Lock);
+	OwningCDM = InOwningCDM;
+	PlayerSession = InForPlayerSession;
+	CDMConfigurations = InCDMConfigurations;
+	GetValuesFromConfigurations();
+	CurrentState = ECDMState::Idle;
+}
+
+//-----------------------------------------------------------------------------
+/**
+ * Registers an event listener to the client.
+ */
+void FClearKeyDRMClient::RegisterEventListener(TWeakPtr<IMediaCDMEventListener, ESPMode::ThreadSafe> InEventListener)
+{
+	FScopeLock lock(&Lock);
+	Listeners.Emplace(InEventListener);
+
+	// Based on the current state we may need to fire events to the new listener right away.
+	if (CurrentState == ECDMState::WaitingForKey)
+	{
+		lock.Unlock();
+		FireEvent(IMediaCDMEventListener::ECDMEventType::KeyRequired);
+	}
+}
+
+//-----------------------------------------------------------------------------
+/**
+ * Unregisters an event listener from the client.
+ */
+void FClearKeyDRMClient::UnregisterEventListener(TWeakPtr<IMediaCDMEventListener, ESPMode::ThreadSafe> InEventListener)
+{
+	FScopeLock lock(&Lock);
+	Listeners.Remove(InEventListener);
+}
+
+//-----------------------------------------------------------------------------
+/**
+ * Fires the given event at all registered event listeners.
+ */
+void FClearKeyDRMClient::FireEvent(IMediaCDMEventListener::ECDMEventType InEvent)
+{
+	TArray<TWeakPtr<IMediaCDMEventListener, ESPMode::ThreadSafe>> CopiedListeners;
+	Lock.Lock();
+	for(int32 i=0; i<Listeners.Num(); ++i)
+	{
+		if (Listeners[i].IsValid())
+		{
+			CopiedListeners.Emplace(Listeners[i]);
+		}
+		else
+		{
+			Listeners.RemoveAt(i);
+			--i;
+		}
+	}
+	Lock.Unlock();
+	TArray<uint8> NoData;
+	TSharedPtr<ElectraCDM::IMediaCDMClient, ESPMode::ThreadSafe> This = AsShared();
+	for(int32 i=0; i<CopiedListeners.Num(); ++i)
+	{
+		TSharedPtr<IMediaCDMEventListener, ESPMode::ThreadSafe> Listener = CopiedListeners[i].Pin();
+		if (Listener.IsValid())
+		{
+			Listener->OnCDMEvent(InEvent, This, nullptr, NoData);
+		}
+	}
+}
+
+//-----------------------------------------------------------------------------
+/**
+ * Prepares the client to fetch a license and fires the event off to the
+ * listeners to start the process.
+ */
+void FClearKeyDRMClient::PrepareLicenses()
+{
+	int32 NumToRequest = PrepareKIDsToRequest();
+	if (NumToRequest)
+	{
+		Lock.Lock();
+		CurrentState = ECDMState::WaitingForKey;
+		Lock.Unlock();
+		FireEvent(IMediaCDMEventListener::ECDMEventType::KeyRequired);
+	}
+	else
+	{
+		Lock.Lock();
+		CurrentState = ECDMState::Ready;
+		Lock.Unlock();
+	}
+}
+
+//-----------------------------------------------------------------------------
+/**
+ * Overrides the license server URL to the given one.
+ * This must happen before calling PrepareLicenses().
+ */
+void FClearKeyDRMClient::SetLicenseServerURL(const FString& InLicenseServerURL)
+{
+	FScopeLock lock(&Lock);
+	LicenseServerURLOverride = InLicenseServerURL;
+}
+
+//-----------------------------------------------------------------------------
+/**
+ * Returns the license server URL to which to issue the license request.
+ */
+void FClearKeyDRMClient::GetLicenseKeyURL(FString& OutLicenseURL)
+{
+	FScopeLock lock(&Lock);
+	// If the URL has been set explicity from the outside return that one.
+	if (LicenseServerURLOverride.IsSet())
+	{
+		OutLicenseURL = LicenseServerURLOverride.GetValue();
+		return;
+	}
+	// Otherwise, when there are several specified through the AdditionalElements
+	// we can return one of them at random.
+	// Which means we take the first one.
+	if (LicenseServerURLsFromConfigs.Num())
+	{
+		OutLicenseURL = LicenseServerURLsFromConfigs[0];
+		return;
+	}
+	// Nothing set at all. Clear out the URL in case it contains something.
+	OutLicenseURL.Empty();
+}
+
+//-----------------------------------------------------------------------------
+/**
+ * Adds a new KID with license key if the KID is not already known.
+ */
+void FClearKeyDRMClient::AddKeyKID(const FClearKeyKIDKey& InKeyKid)
+{
+	FScopeLock lock(&Lock);
+	for(auto &Key : LicenseKeys)
+	{
+		if (Key.KID == InKeyKid.KID)
+		{
+			return;
+		}
+	}
+	LicenseKeys.Add(InKeyKid);
+}
+
+//-----------------------------------------------------------------------------
+/**
+ * Adds a list of new KIDs with license keys when the KID is not already known.
+ */
+void FClearKeyDRMClient::AddKeyKIDs(const TArray<FClearKeyKIDKey>& InKeyKids)
+{
+	for(auto &KeyKid : InKeyKids)
+	{
+		AddKeyKID(KeyKid);
+	}
+}
+
+//-----------------------------------------------------------------------------
+/**
+ * Prepares the list of KIDs for which a license must be obtained.
+ * Licenses the CDM already has will not be requested again.
+ */
+int32 FClearKeyDRMClient::PrepareKIDsToRequest()
+{
+	// We need to get all the KIDs for which we (may) need to acquire a license.
+	FScopeLock lock(&Lock);
+	PendingRequiredKIDs.Empty();
+	check(CDMConfigurations.Num());
+	TSharedPtr<FClearKeyCDM, ESPMode::ThreadSafe> CDM = OwningCDM.Pin();
+	if (CDMConfigurations.Num())
+	{
+		for(int32 nCfg=0; nCfg<CDMConfigurations.Num(); ++nCfg)
+		{
+			for(int32 nKIDs=0; nKIDs<CDMConfigurations[nCfg].DefaultKIDs.Num(); ++nKIDs)
+			{
+				if (CDMConfigurations[nCfg].DefaultKIDs[nKIDs].Len())
+				{
+					// Check if the CDM already has a key for this session's KID.
+					FString KID = StripDashesFromKID(CDMConfigurations[nCfg].DefaultKIDs[nKIDs]);
+					TArray<uint8> BinKID;
+					FClearKeyKIDKey KeyKid;
+					ConvertKIDToBin(BinKID, KID);
+					if (CDM.IsValid() && CDM->GetPlayerSessionKey(KeyKid, PlayerSession, BinKID))
+					{
+						AddKeyKID(KeyKid);
+						continue;
+					}
+					PendingRequiredKIDs.AddUnique(ConvertKIDToBase64(KID));
+				}
+			}
+		}
+	}
+	return PendingRequiredKIDs.Num();
+}
+
+//-----------------------------------------------------------------------------
+/**
+ * Returns the information necessary to make the license request.
+ * This includes the system specific blob of data as well as the HTTP method
+ * to use and additioanl headers, like the "Content-Type: " header.
+ */
+void FClearKeyDRMClient::GetLicenseKeyRequestData(TArray<uint8>& OutKeyRequestData, FString& OutHttpMethod, TArray<FString>& OutHttpHeaders, uint32& OutFlags)
+{
+	FScopeLock lock(&Lock);
+	// The request is a JSON string like here: https://w3c.github.io/encrypted-media/index.html#example-1
+	FString rq1 = TEXT("{\"kids\":[");
+	FString rq2 = TEXT("],\"type\":\"temporary\"}");
+	FString RequestJSON = rq1;
+	for(int32 i=0; i<PendingRequiredKIDs.Num(); ++i)
+	{
+		RequestJSON += FString::Printf(TEXT("\"%s\""), *PendingRequiredKIDs[i]);
+		if (i+1 < PendingRequiredKIDs.Num())
+		RequestJSON.AppendChar(TCHAR(','));
+	}
+	RequestJSON += rq2;
+
+	OutHttpMethod = TEXT("POST");
+	OutHttpHeaders.Emplace(TEXT("Content-Type: application/json"));
+	StringToArray(OutKeyRequestData, RequestJSON);
+	// We allow the use of custom key storage in case there is no license server URL.
+	OutFlags = EDRMClientFlags::EDRMFlg_AllowCustomKeyStorage;
+}
+
+//-----------------------------------------------------------------------------
+/**
+ * Parses the license key response for keys and provides them to the
+ * decrypter instances.
+ */
+ECDMError FClearKeyDRMClient::SetLicenseKeyResponseData(void* InEventId, int32 HttpResponseCode, const TArray<uint8>& InKeyResponseData)
+{
+	// The response is a JSON Web Key (JWK) as per RFC-7518
+	bool bSuccess = true;
+	TArray<FClearKeyKIDKey> NewLicenseKeys;
+	LastErrorMsg.Empty();
+	if (HttpResponseCode == 200)
+	{
+#if ENABLE_LEGACY_RAWKEY_OVERRIDE
+		if (InKeyResponseData.Num() == 16)
+		{
+			FScopeLock lock(&Lock);
+			for(auto &KID : PendingRequiredKIDs)
+			{
+				FClearKeyKIDKey NewKey;
+				if (Base64UrlDecode(NewKey.KID, KID))
+				{
+					NewKey.Key = InKeyResponseData;
+					NewLicenseKeys.Emplace(MoveTemp(NewKey));
+				}
+			}
+		}
+		else
+#endif
+		{
+			FString JsonString = ArrayToString(InKeyResponseData);
+			TSharedRef<TJsonReader<>> Reader = TJsonReaderFactory<>::Create(JsonString);
+			TSharedPtr<FJsonObject> KeyResponse;
+			if (FJsonSerializer::Deserialize(Reader, KeyResponse))
+			{
+				TArray<TSharedPtr<FJsonValue>> Keys;
+				Keys = KeyResponse->GetArrayField("keys");
+				for(int32 i=0; i<Keys.Num(); ++i)
+				{
+					auto Key = Keys[i]->AsObject();
+					FString KTY = Key->GetStringField("kty");	// Mandatory (see RFC-7518 Section 6.1), must be "oct"
+					if (KTY.Equals(TEXT("oct")))
+					{
+						FString KID = Key->GetStringField("kid");
+						FString K = Key->GetStringField("k");
+						if (KID.Len() && K.Len())
+						{
+							FClearKeyKIDKey NewKey;
+							bool bOk = Base64UrlDecode(NewKey.KID, KID);
+							if (bOk)
+							{
+								bOk = Base64UrlDecode(NewKey.Key, K);
+							}
+							if (bOk)
+							{
+								NewLicenseKeys.Emplace(MoveTemp(NewKey));
+							}
+							else
+							{
+								LastErrorMsg = TEXT("Could not base64 decode k or kid in response");
+								bSuccess = false;
+								break;
+							}
+						}
+					}
+					else
+					{
+						LastErrorMsg = TEXT("Unexpected key type (kty) in response");
+						bSuccess = false;
+						break;
+					}
+				}
+			}
+			else
+			{
+				LastErrorMsg = TEXT("Failed to parse license response");
+				bSuccess = false;
+			}
+		}
+
+		if (bSuccess)
+		{
+			TSharedPtr<FClearKeyCDM, ESPMode::ThreadSafe> CDM = OwningCDM.Pin();
+			if (CDM.IsValid())
+			{
+				CDM->AddPlayerSessionKeys(PlayerSession, NewLicenseKeys);
+			}
+
+			AddKeyKIDs(NewLicenseKeys);
+			Lock.Lock();
+			CurrentState = ECDMState::Ready;
+			Lock.Unlock();
+
+			UpdateKeyWithDecrypters();
+			return ECDMError::Success;
+		}
+	}
+	else
+	{
+		LastErrorMsg = FString::Printf(TEXT("Received bad license key response. HTTP code %d"), HttpResponseCode);
+	}
+	CurrentState = ECDMState::InvalidKey;
+	UpdateStateWithDecrypters(CurrentState);
+	return ECDMError::Failure;
+}
+
+//-----------------------------------------------------------------------------
+/**
+ * Creates a decrypter instance.
+ * If the license keys have not been obtained yet the decrypter will not be
+ * usable until the keys arrive.
+ * One decrypter instance is created per elementary stream to decode.
+ */
+ECDMError FClearKeyDRMClient::CreateDecrypter(TSharedPtr<IMediaCDMDecrypter, ESPMode::ThreadSafe>& OutDecrypter, const FString& InMimeType)
+{
+	TSharedPtr<FClearKeyDRMDecrypter, ESPMode::ThreadSafe> NewDec = MakeShared<FClearKeyDRMDecrypter, ESPMode::ThreadSafe>();
+	NewDec->Initialize(InMimeType);
+
+	FScopeLock lock(&Lock);
+
+	// The initial state of the decrypter is the same as the one of the client.
+	NewDec->SetState(CurrentState);
+	NewDec->SetLastErrorMessage(LastErrorMsg);
+	// If ready set the key with the decrypter.
+	if (CurrentState == ECDMState::Ready)
+	{
+		NewDec->SetLicenseKeys(LicenseKeys);
+	}
+
+	RemoveStaleDecrypters();
+	Decrypters.Emplace(NewDec);
+	OutDecrypter = NewDec;
+	return ECDMError::Success;
+}
+
+
+//-----------------------------------------------------------------------------
+/**
+ * Removes decrypters that the application no longer uses.
+ */
+void FClearKeyDRMClient::RemoveStaleDecrypters()
+{
+	for(int32 i=0; i<Decrypters.Num(); ++i)
+	{
+		if (!Decrypters[i].IsValid())
+		{
+			Decrypters.RemoveAt(i);
+			--i;
+		}
+	}
+}
+
+//-----------------------------------------------------------------------------
+/**
+ * Updates all this client's decrypters with the new set of license keys.
+ */
+void FClearKeyDRMClient::UpdateKeyWithDecrypters()
+{
+	FScopeLock lock(&Lock);
+	RemoveStaleDecrypters();
+	for(int32 i=0; i<Decrypters.Num(); ++i)
+	{
+		TSharedPtr<FClearKeyDRMDecrypter, ESPMode::ThreadSafe> Decrypter = Decrypters[i].Pin();
+		Decrypter->SetLicenseKeys(LicenseKeys);
+	}
+}
+
+//-----------------------------------------------------------------------------
+/**
+ * Sets the state of all this client's decrypters to the given state.
+ * This is called in cases of license errors to make all instances fail.
+ */
+void FClearKeyDRMClient::UpdateStateWithDecrypters(ECDMState InNewState)
+{
+	FScopeLock lock(&Lock);
+	RemoveStaleDecrypters();
+	for(int32 i=0; i<Decrypters.Num(); ++i)
+	{
+		TSharedPtr<FClearKeyDRMDecrypter, ESPMode::ThreadSafe> Decrypter = Decrypters[i].Pin();
+		Decrypter->SetLastErrorMessage(LastErrorMsg);
+		Decrypter->SetState(InNewState);
+	}
+}
+
+//-----------------------------------------------------------------------------
+/**
+ * Extracts relevant information from the AdditionalElements
+ */
+void FClearKeyDRMClient::GetValuesFromConfigurations()
+{
+	// For ClearKey we do not expect to get much additional elements.
+	// While there could be a 'pssh' box it would not contain anything
+	// besides the scheme id and the KIDs (in box version 1), which we
+	// expect to have already been set up in the configuration.
+	// What we need is the license server URL or URLs.
+	FScopeLock lock(&Lock);
+	for(auto &Config : CDMConfigurations)
+	{
+		TSharedRef<TJsonReader<>> Reader = TJsonReaderFactory<>::Create(Config.AdditionalElements);
+		TSharedPtr<FJsonObject> ConfigJSON;
+		if (FJsonSerializer::Deserialize(Reader, ConfigJSON))
+		{
+			// Try dashif:laurl first.
+			if (GetURLsFrom(LicenseServerURLsFromConfigs, ConfigJSON, ClearKeyCDM::UrlAttrib_dashif_laurl, ClearKeyCDM::UrlAttrib_laurl))
+			{
+				continue;
+			}
+			// Then the deprecated clearkey:Laurl
+			else if (GetURLsFrom(LicenseServerURLsFromConfigs, ConfigJSON, ClearKeyCDM::UrlAttrib_clearkey_laurl, ClearKeyCDM::UrlAttrib_Laurl))
+			{
+				continue;
+			}
+		}
+		else
+		{
+			LastErrorMsg = TEXT("Could not parse additional configuration element.");
+			CurrentState = ECDMState::ConfigurationError;
+			return;
+		}
+	}
+}
+
+//-----------------------------------------------------------------------------
+/**
+ * Tries to get URLs from a JSON value, which could either be a string, an object
+ * with additional properties, or an array of the same.
+ */
+bool FClearKeyDRMClient::GetURLsFrom(TArray<FString>& OutURLs, const TSharedPtr<FJsonValue>& JSONValue, const FString& PropertyName, const FString& AltPropertyName)
+{
+	// Try string first
+	FString String;
+	if (JSONValue->TryGetString(String))
+	{
+		OutURLs.AddUnique(MoveTemp(String));
+		return true;
+	}
+	// Try object next. This is when the URL has additional properties.
+	const TSharedPtr<FJsonObject>* Object = nullptr;
+	if (JSONValue->TryGetObject(Object))
+	{
+		if ((*Object)->TryGetStringField(ClearKeyCDM::JSONTextPropertyName, String))
+		{
+			OutURLs.AddUnique(MoveTemp(String));
+			return true;
+		}
+		// TBD: should we recurse into the object? For now let's not.
+		//return GetURLsFrom(OutURLs, *Object, PropertyName, AltPropertyName);
+	}
+	// Finally try array
+	const TArray<TSharedPtr<FJsonValue>>* Array = nullptr;
+	if (JSONValue->TryGetArray(Array))
+	{
+		bool bAny = false;
+		for(auto &Val : *Array)
+		{
+			if (GetURLsFrom(OutURLs, Val, PropertyName, AltPropertyName))
+			{
+				bAny = true;
+			}
+		}
+		return bAny;
+	}
+	return false;
+}
+
+//-----------------------------------------------------------------------------
+/**
+ * Tries to get URLs from a JSON object, which could either be a string, an object
+ * with additional properties, or an array of the same.
+ */
+bool FClearKeyDRMClient::GetURLsFrom(TArray<FString>& OutURLs, const TSharedPtr<FJsonObject>& JSON, const FString& PropertyName, const FString& AltPropertyName)
+{
+	TArray<FString> StringArray;
+	FString String;
+	// Try single URL without attributes first.
+	if (JSON->TryGetStringField(PropertyName, String) || (AltPropertyName.Len() && JSON->TryGetStringField(AltPropertyName, String)))
+	{
+		OutURLs.AddUnique(MoveTemp(String));
+		return true;
+	}
+	// Try as array of URLs without attributes.
+	else if (JSON->TryGetStringArrayField(PropertyName, StringArray) || (AltPropertyName.Len() && JSON->TryGetStringArrayField(AltPropertyName, StringArray)))
+	{
+		for(auto &It : StringArray)
+		{
+			OutURLs.AddUnique(It);
+		}
+		return true;
+	}
+	// URL may have additional properties, turning it from a string or string array into an object or array of objects.
+	const TSharedPtr<FJsonObject>* Object = nullptr;
+	// Try single object. Attributes are of no interest at the moment.
+	if (JSON->TryGetObjectField(PropertyName, Object) || (AltPropertyName.Len() && JSON->TryGetObjectField(AltPropertyName, Object)))
+	{
+		if ((*Object)->TryGetStringField(ClearKeyCDM::JSONTextPropertyName, String))
+		{
+			OutURLs.AddUnique(MoveTemp(String));
+			return true;
+		}
+	}
+	// Finally try as array of objects.
+	const TArray<TSharedPtr<FJsonValue>>* ObjectArray = nullptr;
+	if (JSON->TryGetArrayField(PropertyName, ObjectArray) || (AltPropertyName.Len() && JSON->TryGetArrayField(AltPropertyName, ObjectArray)))
+	{
+		bool bAny = false;
+		for(auto &Val : *ObjectArray)
+		{
+			if (GetURLsFrom(OutURLs, Val, PropertyName, AltPropertyName))
+			{
+				bAny = true;
+			}
+		}
+		return bAny;
+	}
+	return false;
+}
+
+
+
+
+//-----------------------------------------------------------------------------
+/**
+ * Creates a new decrypter instance.
+ */
+FClearKeyDRMDecrypter::FClearKeyDRMDecrypter()
+{
+}
+
+//-----------------------------------------------------------------------------
+/**
+ * Destroys a decrypter instance.
+ */
+FClearKeyDRMDecrypter::~FClearKeyDRMDecrypter()
+{
+}
+
+//-----------------------------------------------------------------------------
+/**
+ * Initializes a decrypter instance to default state.
+ */
+void FClearKeyDRMDecrypter::Initialize(const FString& InMimeType)
+{
+	FScopeLock lock(&Lock);
+	CurrentState = ECDMState::Idle;
+	KeyDecrypters.Empty();
+}
+
+//-----------------------------------------------------------------------------
+/**
+ * Updates the valid license keys with this decrypter instance.
+ * All currently active keys are removed and replaced with the new ones.
+ */
+void FClearKeyDRMDecrypter::SetLicenseKeys(const TArray<FClearKeyKIDKey>& InLicenseKeys)
+{
+	FScopeLock lock(&Lock);
+	LicenseKeys = InLicenseKeys;
+	Reinitialize();
+}
+
+//-----------------------------------------------------------------------------
+/**
+ * Returns the current decrypter's state.
+ */
+ECDMState FClearKeyDRMDecrypter::GetState()
+{
+	FScopeLock lock(&Lock);
+	// This is the amalgamated state of all the internal decrypters per
+	// active key.
+	return CurrentState;
+}
+
+//-----------------------------------------------------------------------------
+/**
+ * Returns the most recent error message.
+ */
+FString FClearKeyDRMDecrypter::GetLastErrorMessage()
+{
+	FScopeLock lock(&Lock);
+	return LastErrorMsg;
+}
+
+//-----------------------------------------------------------------------------
+/**
+ * Sets a new state to this decrypter and all its currently active key decrypters.
+ */
+void FClearKeyDRMDecrypter::SetState(ECDMState InNewState)
+{
+	FScopeLock lock(&Lock);
+	CurrentState = InNewState;
+	for(int32 i=0; i<KeyDecrypters.Num(); ++i)
+	{
+		KeyDecrypters[i].State = InNewState;
+	}
+}
+
+//-----------------------------------------------------------------------------
+/**
+ * Updates the last error message.
+ */
+void FClearKeyDRMDecrypter::SetLastErrorMessage(const FString& InNewErrorMessage)
+{
+	FScopeLock lock(&Lock);
+	LastErrorMsg = InNewErrorMessage;
+}
+
+
+//-----------------------------------------------------------------------------
+/**
+ * Called by the application with PSSH box data to update the current set of
+ * key IDs when key rotation is used.
+ */
+ECDMError FClearKeyDRMDecrypter::UpdateInitDataFromPSSH(const TArray<uint8>& InPSSHData)
+{
+	return ECDMError::NotSupported;
+}
+
+ECDMError FClearKeyDRMDecrypter::UpdateInitDataFromMultiplePSSH(const TArray<TArray<uint8>>& InPSSHData)
+{
+	return ECDMError::NotSupported;
+}
+
+//-----------------------------------------------------------------------------
+/**
+ * Update from a URL and additional scheme specific elements.
+ */
+ECDMError FClearKeyDRMDecrypter::UpdateFromURL(const FString& InURL, const FString& InAdditionalElements)
+{
+	return ECDMError::NotSupported;
+}
+
+//-----------------------------------------------------------------------------
+/**
+ * Locates the decrypter for the given key ID.
+ */
+TSharedPtr<ElectraCDM::IStreamDecrypterAES128, ESPMode::ThreadSafe> FClearKeyDRMDecrypter::GetDecrypterForKID(const TArray<uint8>& KID)
+{
+	// Note: The critical section must be locked already!
+	for(int32 i=0; i<KeyDecrypters.Num(); ++i)
+	{
+		if (KID == KeyDecrypters[i].KIDKey.KID)
+		{
+			// Decrypter needs to be ready.
+			if (KeyDecrypters[i].State == ECDMState::Ready)
+			{
+				return KeyDecrypters[i].Decrypter;
+			}
+		}
+	}
+	return nullptr;
+}
+
+//-----------------------------------------------------------------------------
+/**
+ * Reinitializes the decrypter to its starting state.
+ */
+void FClearKeyDRMDecrypter::Reinitialize()
+{
+	FScopeLock lock(&Lock);
+
+	KeyDecrypters.Empty();
+	LastErrorMsg.Empty();
+
+	ElectraCDM::IStreamDecrypterAES128::EResult Result;
+	CurrentState = ECDMState::Ready;
+	for(int32 i=0; i<LicenseKeys.Num(); ++i)
+	{
+		FKeyDecrpyter& kd = KeyDecrypters.AddDefaulted_GetRef();
+		kd.KIDKey = LicenseKeys[i];
+		kd.Decrypter = ElectraCDM::IStreamDecrypterAES128::Create();
+		Result = kd.Decrypter->CTRInit(kd.KIDKey.Key);
+		if (Result == ElectraCDM::IStreamDecrypterAES128::EResult::Ok)
+		{
+			kd.State = ECDMState::Ready;
+		}
+		else
+		{
+			kd.State = ECDMState::InvalidKey;
+			CurrentState = ECDMState::InvalidKey;
+			LastErrorMsg = TEXT("Invalid key");
+		}
+	}
+}
+
+//-----------------------------------------------------------------------------
+/**
+ * Decrypts data in place according to the encrypted sample information.
+ */
+ECDMError FClearKeyDRMDecrypter::DecryptInPlace(uint8* InOutData, int32 InNumDataBytes, const FMediaCDMSampleInfo& InSampleInfo)
+{
+	FScopeLock lock(&Lock);
+	LastErrorMsg.Empty();
+	TSharedPtr<ElectraCDM::IStreamDecrypterAES128, ESPMode::ThreadSafe> Decrypter = GetDecrypterForKID(InSampleInfo.DefaultKID);
+	if (Decrypter.IsValid())
+	{
+		if (Decrypter->CTRSetIV(InSampleInfo.IV) != ElectraCDM::IStreamDecrypterAES128::EResult::Ok)
+		{
+			LastErrorMsg = TEXT("Bad IV");
+			return ECDMError::Failure;
+		}
+		if (InSampleInfo.SubSamples.Num() == 0)
+		{
+			Decrypter->CTRDecryptInPlace(InOutData, InNumDataBytes);
+			return ECDMError::Success;
+		}
+		else
+		{
+			for(int32 i=0; i<InSampleInfo.SubSamples.Num(); ++i)
+			{
+				InOutData += InSampleInfo.SubSamples[i].NumClearBytes;
+				if (InSampleInfo.SubSamples[i].NumEncryptedBytes)
+				{
+					Decrypter->CTRDecryptInPlace(InOutData, InSampleInfo.SubSamples[i].NumEncryptedBytes);
+				}
+				InOutData += InSampleInfo.SubSamples[i].NumEncryptedBytes;
+			}
+			return ECDMError::Success;
+		}
+	}
+	LastErrorMsg = TEXT("No valid decrypter found for KID");
+	return ECDMError::Failure;
+}
+
+
+bool FClearKeyDRMDecrypter::IsBlockStreamDecrypter()
+{
+	return false;
+}
+
+ECDMError FClearKeyDRMDecrypter::BlockStreamDecryptStart(IStreamDecryptHandle*& OutStreamDecryptContext)
+{
+	FScopeLock lock(&Lock);
+	OutStreamDecryptContext = nullptr;
+	LastErrorMsg = TEXT("Not a block stream decrypter");
+	return ECDMError::CipherModeMismatch;
+}
+ECDMError FClearKeyDRMDecrypter::BlockStreamDecryptInPlace(IStreamDecryptHandle* InOutStreamDecryptContext, int32& OutNumBytesDecrypted, uint8* InOutData, int32 InNumDataBytes, const FMediaCDMSampleInfo& InSampleInfo, bool bIsLastBlock)
+{
+	FScopeLock lock(&Lock);
+	LastErrorMsg = TEXT("Not a block stream decrypter");
+	return ECDMError::CipherModeMismatch;
+}
+ECDMError FClearKeyDRMDecrypter::BlockStreamDecryptEnd(IStreamDecryptHandle* InStreamDecryptContext)
+{
+	FScopeLock lock(&Lock);
+	LastErrorMsg = TEXT("Not a block stream decrypter");
+	return ECDMError::CipherModeMismatch;
+}
+
+
+}
diff --git a/Engine/Plugins/Media/ElectraCDM/Source/ElectraCDM/Private/ClearKey/ClearKeyCDM.h b/Engine/Plugins/Media/ElectraCDM/Source/ElectraCDM/Private/ClearKey/ClearKeyCDM.h
new file mode 100644
--- /dev/null
+++ b/Engine/Plugins/Media/ElectraCDM/Source/ElectraCDM/Private/ClearKey/ClearKeyCDM.h
@@ -0,0 +1,49 @@
+// Copyright Epic Games, Inc. All Rights Reserved.
+
+#pragma once
+
+#include "ElectraCDM.h"
+#include "ElectraCDMSystem.h"
+
+namespace ElectraCDM
+{
+
+class IClearKeyCDM : public IMediaCDMSystem
+{
+public:
+	static void RegisterWith(IMediaCDM& InDRMManager);
+	virtual ~IClearKeyCDM() = default;
+	virtual FString GetLastErrorMessage() = 0;
+	virtual const FString& GetSchemeID() = 0;
+
+	/*
+		Elements that are parsed from a XML document, like:
+			<ContentProtection schemeIdUri="urn:uuid:e2719d58-a985-b3c9-781a-b030af78d30e" value="ClearKey1.0">
+			  <clearkey:Laurl Lic_type="EME-1.0">https://clearkey.example.com/AcquireLicense</clearkey:Laurl>
+			  <dashif:laurl>https://clearkey2.example.com/AcquireLicense</dashif:laurl>
+			  <cenc:pssh>YmFzZTY0IGVuY29kZWQgY29udGVudHMgb2YgkXBzc2iSIGJveCB3aXRoIHRoaXMgU3lzdGVtSUQ=</cenc:pssh>
+			<ContentProtection>
+
+		have these elements converted into JSON like:
+
+			{
+				"Laurl": {
+					"Lic_type": "EME-1.0",
+					"#text": "https://clearkey.example.com/AcquireLicense"
+				},
+				"laurl": "https://clearkey2.example.com/AcquireLicense",
+				"pssh": "YmFzZTY0IGVuY29kZWQgY29udGVudHMgb2YgkXBzc2iSIGJveCB3aXRoIHRoaXMgU3lzdGVtSUQ="
+			}
+
+		where the text of the XML element is put into a property named "#text". XML namespaces are removed in this example
+		(bOutNoNamespaces returned 'true') and no attribute prefix was set.
+	*/
+	virtual void GetCDMCustomJSONPrefixes(FString& OutAttributePrefix, FString& OutTextPropertyName, bool& bOutNoNamespaces) = 0;
+
+	virtual TSharedPtr<IMediaCDMCapabilities, ESPMode::ThreadSafe> GetCDMCapabilities(const FString& InValue, const FString& InAdditionalElements) = 0;
+	virtual ECDMError CreateDRMClient(TSharedPtr<IMediaCDMClient, ESPMode::ThreadSafe>& OutClient, IMediaCDM::IPlayerSession* InForPlayerSession, const TArray<IMediaCDM::FCDMCandidate>& InCandidates) = 0;
+	virtual ECDMError ReleasePlayerSessionKeys(IMediaCDM::IPlayerSession* PlayerSession) = 0;
+};
+
+}
+
diff --git a/Engine/Plugins/Media/ElectraCDM/Source/ElectraCDM/Private/Crypto/StreamCryptoAES128.cpp b/Engine/Plugins/Media/ElectraCDM/Source/ElectraCDM/Private/Crypto/StreamCryptoAES128.cpp
new file mode 100644
--- /dev/null
+++ b/Engine/Plugins/Media/ElectraCDM/Source/ElectraCDM/Private/Crypto/StreamCryptoAES128.cpp
@@ -0,0 +1,580 @@
+// Copyright Epic Games, Inc. All Rights Reserved.
+
+#include "Crypto/StreamCryptoAES128.h"
+
+#include "3rdParty/tiny-AES-c/aes.h"
+
+namespace ElectraCDM
+{
+
+/***************************************************************************************************************************************************/
+/***************************************************************************************************************************************************/
+/***************************************************************************************************************************************************/
+
+class FAES128CTR
+{
+public:
+	void SetKey(const uint8* InKey)
+	{
+		uint32 *RK = RoundKey;
+		RK[0] = GetUInt32(InKey);
+		RK[1] = GetUInt32(InKey + 4);
+		RK[2] = GetUInt32(InKey + 8);
+		RK[3] = GetUInt32(InKey + 12);
+		for(int32 i=0; i<10; ++i)
+		{
+			uint32 temp = RK[3];
+			RK[4] = RK[0] ^ (T2[(temp >> 16) & 255] & 0xff000000) ^ (T3[(temp >> 8) & 255] & 0xff0000) ^ (T0[temp & 255] & 0xff00) ^ (T1[temp >> 24] & 0xff) ^ RoundConstant[i];
+			RK[5] = RK[1] ^ RK[4];
+			RK[6] = RK[2] ^ RK[5];
+			RK[7] = RK[3] ^ RK[6];
+			RK += 4;
+		}
+		CipherPos = 0;
+	}
+	
+	void SetIV(const TArray<uint8>& InIV)
+	{
+		if (InIV.Num() <= 16)
+		{
+			FMemory::Memzero(IV, 16);
+			FMemory::Memcpy(IV, InIV.GetData(), InIV.Num());
+		}
+		CipherPos = 0;
+	}
+
+	void IncrementCounter()
+	{
+		uint8* Ctr = IV + 15;
+		for(int32 i=0; i<8; ++i)
+		{
+			*Ctr = *Ctr + 1;
+			if (*Ctr)
+			{
+				return;
+			}
+			--Ctr;
+		}
+	}
+
+	void Decrypt(uint8* InOutData, int32 InNumBytes)
+	{
+		while(InNumBytes)
+		{
+			if (CipherPos == 0)
+			{
+				Cipher(CurrentCipherBuffer, IV);
+			}
+			int32 maxb = 16 - CipherPos;
+			int32 n = InNumBytes < maxb ? InNumBytes : maxb;
+			for(int32 i=0; i<n; ++i, ++InOutData)
+			{
+				*InOutData = *InOutData ^ CurrentCipherBuffer[CipherPos];
+				if (++CipherPos == 16)
+				{
+					CipherPos = 0;
+					IncrementCounter();
+				}
+			}
+			InNumBytes -= n;
+		}
+	}
+
+private:
+	void Cipher(uint8* Out, const uint8* In) const
+	{
+		uint32 s0, s1, s2, s3, t0, t1, t2, t3;
+		s0 = GetUInt32(In)      ^ RoundKey[0];
+		s1 = GetUInt32(In + 4)  ^ RoundKey[1];
+		s2 = GetUInt32(In + 8)  ^ RoundKey[2];
+		s3 = GetUInt32(In + 12) ^ RoundKey[3];
+		t0 = T0[s0 >> 24] ^ T1[(s1 >> 16) & 255] ^ T2[(s2 >> 8) & 255] ^ T3[s3 & 255] ^ RoundKey[ 4];
+		t1 = T0[s1 >> 24] ^ T1[(s2 >> 16) & 255] ^ T2[(s3 >> 8) & 255] ^ T3[s0 & 255] ^ RoundKey[ 5];
+		t2 = T0[s2 >> 24] ^ T1[(s3 >> 16) & 255] ^ T2[(s0 >> 8) & 255] ^ T3[s1 & 255] ^ RoundKey[ 6];
+		t3 = T0[s3 >> 24] ^ T1[(s0 >> 16) & 255] ^ T2[(s1 >> 8) & 255] ^ T3[s2 & 255] ^ RoundKey[ 7];
+		s0 = T0[t0 >> 24] ^ T1[(t1 >> 16) & 255] ^ T2[(t2 >> 8) & 255] ^ T3[t3 & 255] ^ RoundKey[ 8];
+		s1 = T0[t1 >> 24] ^ T1[(t2 >> 16) & 255] ^ T2[(t3 >> 8) & 255] ^ T3[t0 & 255] ^ RoundKey[ 9];
+		s2 = T0[t2 >> 24] ^ T1[(t3 >> 16) & 255] ^ T2[(t0 >> 8) & 255] ^ T3[t1 & 255] ^ RoundKey[10];
+		s3 = T0[t3 >> 24] ^ T1[(t0 >> 16) & 255] ^ T2[(t1 >> 8) & 255] ^ T3[t2 & 255] ^ RoundKey[11];
+		t0 = T0[s0 >> 24] ^ T1[(s1 >> 16) & 255] ^ T2[(s2 >> 8) & 255] ^ T3[s3 & 255] ^ RoundKey[12];
+		t1 = T0[s1 >> 24] ^ T1[(s2 >> 16) & 255] ^ T2[(s3 >> 8) & 255] ^ T3[s0 & 255] ^ RoundKey[13];
+		t2 = T0[s2 >> 24] ^ T1[(s3 >> 16) & 255] ^ T2[(s0 >> 8) & 255] ^ T3[s1 & 255] ^ RoundKey[14];
+		t3 = T0[s3 >> 24] ^ T1[(s0 >> 16) & 255] ^ T2[(s1 >> 8) & 255] ^ T3[s2 & 255] ^ RoundKey[15];
+		s0 = T0[t0 >> 24] ^ T1[(t1 >> 16) & 255] ^ T2[(t2 >> 8) & 255] ^ T3[t3 & 255] ^ RoundKey[16];
+		s1 = T0[t1 >> 24] ^ T1[(t2 >> 16) & 255] ^ T2[(t3 >> 8) & 255] ^ T3[t0 & 255] ^ RoundKey[17];
+		s2 = T0[t2 >> 24] ^ T1[(t3 >> 16) & 255] ^ T2[(t0 >> 8) & 255] ^ T3[t1 & 255] ^ RoundKey[18];
+		s3 = T0[t3 >> 24] ^ T1[(t0 >> 16) & 255] ^ T2[(t1 >> 8) & 255] ^ T3[t2 & 255] ^ RoundKey[19];
+		t0 = T0[s0 >> 24] ^ T1[(s1 >> 16) & 255] ^ T2[(s2 >> 8) & 255] ^ T3[s3 & 255] ^ RoundKey[20];
+		t1 = T0[s1 >> 24] ^ T1[(s2 >> 16) & 255] ^ T2[(s3 >> 8) & 255] ^ T3[s0 & 255] ^ RoundKey[21];
+		t2 = T0[s2 >> 24] ^ T1[(s3 >> 16) & 255] ^ T2[(s0 >> 8) & 255] ^ T3[s1 & 255] ^ RoundKey[22];
+		t3 = T0[s3 >> 24] ^ T1[(s0 >> 16) & 255] ^ T2[(s1 >> 8) & 255] ^ T3[s2 & 255] ^ RoundKey[23];
+		s0 = T0[t0 >> 24] ^ T1[(t1 >> 16) & 255] ^ T2[(t2 >> 8) & 255] ^ T3[t3 & 255] ^ RoundKey[24];
+		s1 = T0[t1 >> 24] ^ T1[(t2 >> 16) & 255] ^ T2[(t3 >> 8) & 255] ^ T3[t0 & 255] ^ RoundKey[25];
+		s2 = T0[t2 >> 24] ^ T1[(t3 >> 16) & 255] ^ T2[(t0 >> 8) & 255] ^ T3[t1 & 255] ^ RoundKey[26];
+		s3 = T0[t3 >> 24] ^ T1[(t0 >> 16) & 255] ^ T2[(t1 >> 8) & 255] ^ T3[t2 & 255] ^ RoundKey[27];
+		t0 = T0[s0 >> 24] ^ T1[(s1 >> 16) & 255] ^ T2[(s2 >> 8) & 255] ^ T3[s3 & 255] ^ RoundKey[28];
+		t1 = T0[s1 >> 24] ^ T1[(s2 >> 16) & 255] ^ T2[(s3 >> 8) & 255] ^ T3[s0 & 255] ^ RoundKey[29];
+		t2 = T0[s2 >> 24] ^ T1[(s3 >> 16) & 255] ^ T2[(s0 >> 8) & 255] ^ T3[s1 & 255] ^ RoundKey[30];
+		t3 = T0[s3 >> 24] ^ T1[(s0 >> 16) & 255] ^ T2[(s1 >> 8) & 255] ^ T3[s2 & 255] ^ RoundKey[31];
+		s0 = T0[t0 >> 24] ^ T1[(t1 >> 16) & 255] ^ T2[(t2 >> 8) & 255] ^ T3[t3 & 255] ^ RoundKey[32];
+		s1 = T0[t1 >> 24] ^ T1[(t2 >> 16) & 255] ^ T2[(t3 >> 8) & 255] ^ T3[t0 & 255] ^ RoundKey[33];
+		s2 = T0[t2 >> 24] ^ T1[(t3 >> 16) & 255] ^ T2[(t0 >> 8) & 255] ^ T3[t1 & 255] ^ RoundKey[34];
+		s3 = T0[t3 >> 24] ^ T1[(t0 >> 16) & 255] ^ T2[(t1 >> 8) & 255] ^ T3[t2 & 255] ^ RoundKey[35];
+		t0 = T0[s0 >> 24] ^ T1[(s1 >> 16) & 255] ^ T2[(s2 >> 8) & 255] ^ T3[s3 & 255] ^ RoundKey[36];
+		t1 = T0[s1 >> 24] ^ T1[(s2 >> 16) & 255] ^ T2[(s3 >> 8) & 255] ^ T3[s0 & 255] ^ RoundKey[37];
+		t2 = T0[s2 >> 24] ^ T1[(s3 >> 16) & 255] ^ T2[(s0 >> 8) & 255] ^ T3[s1 & 255] ^ RoundKey[38];
+		t3 = T0[s3 >> 24] ^ T1[(s0 >> 16) & 255] ^ T2[(s1 >> 8) & 255] ^ T3[s2 & 255] ^ RoundKey[39];
+		s0 = (T2[(t0 >> 24)] & 0xff000000) ^ (T3[(t1 >> 16) & 255] & 0xff0000) ^ (T0[(t2 >> 8) & 255] & 0xff00) ^ (T1[t3 & 255] & 0xff) ^ RoundKey[40];
+		s1 = (T2[(t1 >> 24)] & 0xff000000) ^ (T3[(t2 >> 16) & 255] & 0xff0000) ^ (T0[(t3 >> 8) & 255] & 0xff00) ^ (T1[t0 & 255] & 0xff) ^ RoundKey[41];
+		s2 = (T2[(t2 >> 24)] & 0xff000000) ^ (T3[(t3 >> 16) & 255] & 0xff0000) ^ (T0[(t0 >> 8) & 255] & 0xff00) ^ (T1[t1 & 255] & 0xff) ^ RoundKey[42];
+		s3 = (T2[(t3 >> 24)] & 0xff000000) ^ (T3[(t0 >> 16) & 255] & 0xff0000) ^ (T0[(t1 >> 8) & 255] & 0xff00) ^ (T1[t2 & 255] & 0xff) ^ RoundKey[43];
+		PutUInt32(Out,      s0);
+		PutUInt32(Out + 4,  s1);
+		PutUInt32(Out + 8,  s2);
+		PutUInt32(Out + 12, s3);
+	}
+
+	static uint32 GetUInt32(const uint8* InFrom)
+	{
+		return ((uint32)InFrom[0] << 24) | ((uint32)InFrom[1] << 16) | ((uint32)InFrom[2] << 8) | (uint32)InFrom[3];
+	}
+
+	static void PutUInt32(uint8* OutTo, uint32 v)
+	{
+		OutTo[0] = uint8(v >> 24);
+		OutTo[1] = uint8(v >> 16);
+		OutTo[2] = uint8(v >> 8);
+		OutTo[3] = uint8(v);
+	}
+
+	static const uint32	T0[256];
+	static const uint32	T1[256];
+	static const uint32	T2[256];
+	static const uint32	T3[256];
+	static const uint32	RoundConstant[10];
+	uint32 RoundKey[44];
+	uint8 IV[16];
+	int32 CipherPos = 0;
+	uint8 CurrentCipherBuffer[16];
+};
+const uint32 FAES128CTR::T0[256] =
+{
+	0xc66363a5U, 0xf87c7c84U, 0xee777799U, 0xf67b7b8dU, 0xfff2f20dU, 0xd66b6bbdU, 0xde6f6fb1U, 0x91c5c554U,
+	0x60303050U, 0x02010103U, 0xce6767a9U, 0x562b2b7dU, 0xe7fefe19U, 0xb5d7d762U, 0x4dababe6U, 0xec76769aU,
+	0x8fcaca45U, 0x1f82829dU, 0x89c9c940U, 0xfa7d7d87U, 0xeffafa15U, 0xb25959ebU, 0x8e4747c9U, 0xfbf0f00bU,
+	0x41adadecU, 0xb3d4d467U, 0x5fa2a2fdU, 0x45afafeaU, 0x239c9cbfU, 0x53a4a4f7U, 0xe4727296U, 0x9bc0c05bU,
+	0x75b7b7c2U, 0xe1fdfd1cU, 0x3d9393aeU, 0x4c26266aU, 0x6c36365aU, 0x7e3f3f41U, 0xf5f7f702U, 0x83cccc4fU,
+	0x6834345cU, 0x51a5a5f4U, 0xd1e5e534U, 0xf9f1f108U, 0xe2717193U, 0xabd8d873U, 0x62313153U, 0x2a15153fU,
+	0x0804040cU, 0x95c7c752U, 0x46232365U, 0x9dc3c35eU, 0x30181828U, 0x379696a1U, 0x0a05050fU, 0x2f9a9ab5U,
+	0x0e070709U, 0x24121236U, 0x1b80809bU, 0xdfe2e23dU, 0xcdebeb26U, 0x4e272769U, 0x7fb2b2cdU, 0xea75759fU,
+	0x1209091bU, 0x1d83839eU, 0x582c2c74U, 0x341a1a2eU, 0x361b1b2dU, 0xdc6e6eb2U, 0xb45a5aeeU, 0x5ba0a0fbU,
+	0xa45252f6U, 0x763b3b4dU, 0xb7d6d661U, 0x7db3b3ceU, 0x5229297bU, 0xdde3e33eU, 0x5e2f2f71U, 0x13848497U,
+	0xa65353f5U, 0xb9d1d168U, 0x00000000U, 0xc1eded2cU, 0x40202060U, 0xe3fcfc1fU, 0x79b1b1c8U, 0xb65b5bedU,
+	0xd46a6abeU, 0x8dcbcb46U, 0x67bebed9U, 0x7239394bU, 0x944a4adeU, 0x984c4cd4U, 0xb05858e8U, 0x85cfcf4aU,
+	0xbbd0d06bU, 0xc5efef2aU, 0x4faaaae5U, 0xedfbfb16U, 0x864343c5U, 0x9a4d4dd7U, 0x66333355U, 0x11858594U,
+	0x8a4545cfU, 0xe9f9f910U, 0x04020206U, 0xfe7f7f81U, 0xa05050f0U, 0x783c3c44U, 0x259f9fbaU, 0x4ba8a8e3U,
+	0xa25151f3U, 0x5da3a3feU, 0x804040c0U, 0x058f8f8aU, 0x3f9292adU, 0x219d9dbcU, 0x70383848U, 0xf1f5f504U,
+	0x63bcbcdfU, 0x77b6b6c1U, 0xafdada75U, 0x42212163U, 0x20101030U, 0xe5ffff1aU, 0xfdf3f30eU, 0xbfd2d26dU,
+	0x81cdcd4cU, 0x180c0c14U, 0x26131335U, 0xc3ecec2fU, 0xbe5f5fe1U, 0x359797a2U, 0x884444ccU, 0x2e171739U,
+	0x93c4c457U, 0x55a7a7f2U, 0xfc7e7e82U, 0x7a3d3d47U, 0xc86464acU, 0xba5d5de7U, 0x3219192bU, 0xe6737395U,
+	0xc06060a0U, 0x19818198U, 0x9e4f4fd1U, 0xa3dcdc7fU, 0x44222266U, 0x542a2a7eU, 0x3b9090abU, 0x0b888883U,
+	0x8c4646caU, 0xc7eeee29U, 0x6bb8b8d3U, 0x2814143cU, 0xa7dede79U, 0xbc5e5ee2U, 0x160b0b1dU, 0xaddbdb76U,
+	0xdbe0e03bU, 0x64323256U, 0x743a3a4eU, 0x140a0a1eU, 0x924949dbU, 0x0c06060aU, 0x4824246cU, 0xb85c5ce4U,
+	0x9fc2c25dU, 0xbdd3d36eU, 0x43acacefU, 0xc46262a6U, 0x399191a8U, 0x319595a4U, 0xd3e4e437U, 0xf279798bU,
+	0xd5e7e732U, 0x8bc8c843U, 0x6e373759U, 0xda6d6db7U, 0x018d8d8cU, 0xb1d5d564U, 0x9c4e4ed2U, 0x49a9a9e0U,
+	0xd86c6cb4U, 0xac5656faU, 0xf3f4f407U, 0xcfeaea25U, 0xca6565afU, 0xf47a7a8eU, 0x47aeaee9U, 0x10080818U,
+	0x6fbabad5U, 0xf0787888U, 0x4a25256fU, 0x5c2e2e72U, 0x381c1c24U, 0x57a6a6f1U, 0x73b4b4c7U, 0x97c6c651U,
+	0xcbe8e823U, 0xa1dddd7cU, 0xe874749cU, 0x3e1f1f21U, 0x964b4bddU, 0x61bdbddcU, 0x0d8b8b86U, 0x0f8a8a85U,
+	0xe0707090U, 0x7c3e3e42U, 0x71b5b5c4U, 0xcc6666aaU, 0x904848d8U, 0x06030305U, 0xf7f6f601U, 0x1c0e0e12U,
+	0xc26161a3U, 0x6a35355fU, 0xae5757f9U, 0x69b9b9d0U, 0x17868691U, 0x99c1c158U, 0x3a1d1d27U, 0x279e9eb9U,
+	0xd9e1e138U, 0xebf8f813U, 0x2b9898b3U, 0x22111133U, 0xd26969bbU, 0xa9d9d970U, 0x078e8e89U, 0x339494a7U,
+	0x2d9b9bb6U, 0x3c1e1e22U, 0x15878792U, 0xc9e9e920U, 0x87cece49U, 0xaa5555ffU, 0x50282878U, 0xa5dfdf7aU,
+	0x038c8c8fU, 0x59a1a1f8U, 0x09898980U, 0x1a0d0d17U, 0x65bfbfdaU, 0xd7e6e631U, 0x844242c6U, 0xd06868b8U,
+	0x824141c3U, 0x299999b0U, 0x5a2d2d77U, 0x1e0f0f11U, 0x7bb0b0cbU, 0xa85454fcU, 0x6dbbbbd6U, 0x2c16163aU,
+};
+const uint32 FAES128CTR::T1[256] =
+{
+	0xa5c66363U, 0x84f87c7cU, 0x99ee7777U, 0x8df67b7bU, 0x0dfff2f2U, 0xbdd66b6bU, 0xb1de6f6fU, 0x5491c5c5U,
+	0x50603030U, 0x03020101U, 0xa9ce6767U, 0x7d562b2bU, 0x19e7fefeU, 0x62b5d7d7U, 0xe64dababU, 0x9aec7676U,
+	0x458fcacaU, 0x9d1f8282U, 0x4089c9c9U, 0x87fa7d7dU, 0x15effafaU, 0xebb25959U, 0xc98e4747U, 0x0bfbf0f0U,
+	0xec41adadU, 0x67b3d4d4U, 0xfd5fa2a2U, 0xea45afafU, 0xbf239c9cU, 0xf753a4a4U, 0x96e47272U, 0x5b9bc0c0U,
+	0xc275b7b7U, 0x1ce1fdfdU, 0xae3d9393U, 0x6a4c2626U, 0x5a6c3636U, 0x417e3f3fU, 0x02f5f7f7U, 0x4f83ccccU,
+	0x5c683434U, 0xf451a5a5U, 0x34d1e5e5U, 0x08f9f1f1U, 0x93e27171U, 0x73abd8d8U, 0x53623131U, 0x3f2a1515U,
+	0x0c080404U, 0x5295c7c7U, 0x65462323U, 0x5e9dc3c3U, 0x28301818U, 0xa1379696U, 0x0f0a0505U, 0xb52f9a9aU,
+	0x090e0707U, 0x36241212U, 0x9b1b8080U, 0x3ddfe2e2U, 0x26cdebebU, 0x694e2727U, 0xcd7fb2b2U, 0x9fea7575U,
+	0x1b120909U, 0x9e1d8383U, 0x74582c2cU, 0x2e341a1aU, 0x2d361b1bU, 0xb2dc6e6eU, 0xeeb45a5aU, 0xfb5ba0a0U,
+	0xf6a45252U, 0x4d763b3bU, 0x61b7d6d6U, 0xce7db3b3U, 0x7b522929U, 0x3edde3e3U, 0x715e2f2fU, 0x97138484U,
+	0xf5a65353U, 0x68b9d1d1U, 0x00000000U, 0x2cc1ededU, 0x60402020U, 0x1fe3fcfcU, 0xc879b1b1U, 0xedb65b5bU,
+	0xbed46a6aU, 0x468dcbcbU, 0xd967bebeU, 0x4b723939U, 0xde944a4aU, 0xd4984c4cU, 0xe8b05858U, 0x4a85cfcfU,
+	0x6bbbd0d0U, 0x2ac5efefU, 0xe54faaaaU, 0x16edfbfbU, 0xc5864343U, 0xd79a4d4dU, 0x55663333U, 0x94118585U,
+	0xcf8a4545U, 0x10e9f9f9U, 0x06040202U, 0x81fe7f7fU, 0xf0a05050U, 0x44783c3cU, 0xba259f9fU, 0xe34ba8a8U,
+	0xf3a25151U, 0xfe5da3a3U, 0xc0804040U, 0x8a058f8fU, 0xad3f9292U, 0xbc219d9dU, 0x48703838U, 0x04f1f5f5U,
+	0xdf63bcbcU, 0xc177b6b6U, 0x75afdadaU, 0x63422121U, 0x30201010U, 0x1ae5ffffU, 0x0efdf3f3U, 0x6dbfd2d2U,
+	0x4c81cdcdU, 0x14180c0cU, 0x35261313U, 0x2fc3ececU, 0xe1be5f5fU, 0xa2359797U, 0xcc884444U, 0x392e1717U,
+	0x5793c4c4U, 0xf255a7a7U, 0x82fc7e7eU, 0x477a3d3dU, 0xacc86464U, 0xe7ba5d5dU, 0x2b321919U, 0x95e67373U,
+	0xa0c06060U, 0x98198181U, 0xd19e4f4fU, 0x7fa3dcdcU, 0x66442222U, 0x7e542a2aU, 0xab3b9090U, 0x830b8888U,
+	0xca8c4646U, 0x29c7eeeeU, 0xd36bb8b8U, 0x3c281414U, 0x79a7dedeU, 0xe2bc5e5eU, 0x1d160b0bU, 0x76addbdbU,
+	0x3bdbe0e0U, 0x56643232U, 0x4e743a3aU, 0x1e140a0aU, 0xdb924949U, 0x0a0c0606U, 0x6c482424U, 0xe4b85c5cU,
+	0x5d9fc2c2U, 0x6ebdd3d3U, 0xef43acacU, 0xa6c46262U, 0xa8399191U, 0xa4319595U, 0x37d3e4e4U, 0x8bf27979U,
+	0x32d5e7e7U, 0x438bc8c8U, 0x596e3737U, 0xb7da6d6dU, 0x8c018d8dU, 0x64b1d5d5U, 0xd29c4e4eU, 0xe049a9a9U,
+	0xb4d86c6cU, 0xfaac5656U, 0x07f3f4f4U, 0x25cfeaeaU, 0xafca6565U, 0x8ef47a7aU, 0xe947aeaeU, 0x18100808U,
+	0xd56fbabaU, 0x88f07878U, 0x6f4a2525U, 0x725c2e2eU, 0x24381c1cU, 0xf157a6a6U, 0xc773b4b4U, 0x5197c6c6U,
+	0x23cbe8e8U, 0x7ca1ddddU, 0x9ce87474U, 0x213e1f1fU, 0xdd964b4bU, 0xdc61bdbdU, 0x860d8b8bU, 0x850f8a8aU,
+	0x90e07070U, 0x427c3e3eU, 0xc471b5b5U, 0xaacc6666U, 0xd8904848U, 0x05060303U, 0x01f7f6f6U, 0x121c0e0eU,
+	0xa3c26161U, 0x5f6a3535U, 0xf9ae5757U, 0xd069b9b9U, 0x91178686U, 0x5899c1c1U, 0x273a1d1dU, 0xb9279e9eU,
+	0x38d9e1e1U, 0x13ebf8f8U, 0xb32b9898U, 0x33221111U, 0xbbd26969U, 0x70a9d9d9U, 0x89078e8eU, 0xa7339494U,
+	0xb62d9b9bU, 0x223c1e1eU, 0x92158787U, 0x20c9e9e9U, 0x4987ceceU, 0xffaa5555U, 0x78502828U, 0x7aa5dfdfU,
+	0x8f038c8cU, 0xf859a1a1U, 0x80098989U, 0x171a0d0dU, 0xda65bfbfU, 0x31d7e6e6U, 0xc6844242U, 0xb8d06868U,
+	0xc3824141U, 0xb0299999U, 0x775a2d2dU, 0x111e0f0fU, 0xcb7bb0b0U, 0xfca85454U, 0xd66dbbbbU, 0x3a2c1616U,
+};
+const uint32 FAES128CTR::T2[256] =
+{
+	0x63a5c663U, 0x7c84f87cU, 0x7799ee77U, 0x7b8df67bU, 0xf20dfff2U, 0x6bbdd66bU, 0x6fb1de6fU, 0xc55491c5U,
+	0x30506030U, 0x01030201U, 0x67a9ce67U, 0x2b7d562bU, 0xfe19e7feU, 0xd762b5d7U, 0xabe64dabU, 0x769aec76U,
+	0xca458fcaU, 0x829d1f82U, 0xc94089c9U, 0x7d87fa7dU, 0xfa15effaU, 0x59ebb259U, 0x47c98e47U, 0xf00bfbf0U,
+	0xadec41adU, 0xd467b3d4U, 0xa2fd5fa2U, 0xafea45afU, 0x9cbf239cU, 0xa4f753a4U, 0x7296e472U, 0xc05b9bc0U,
+	0xb7c275b7U, 0xfd1ce1fdU, 0x93ae3d93U, 0x266a4c26U, 0x365a6c36U, 0x3f417e3fU, 0xf702f5f7U, 0xcc4f83ccU,
+	0x345c6834U, 0xa5f451a5U, 0xe534d1e5U, 0xf108f9f1U, 0x7193e271U, 0xd873abd8U, 0x31536231U, 0x153f2a15U,
+	0x040c0804U, 0xc75295c7U, 0x23654623U, 0xc35e9dc3U, 0x18283018U, 0x96a13796U, 0x050f0a05U, 0x9ab52f9aU,
+	0x07090e07U, 0x12362412U, 0x809b1b80U, 0xe23ddfe2U, 0xeb26cdebU, 0x27694e27U, 0xb2cd7fb2U, 0x759fea75U,
+	0x091b1209U, 0x839e1d83U, 0x2c74582cU, 0x1a2e341aU, 0x1b2d361bU, 0x6eb2dc6eU, 0x5aeeb45aU, 0xa0fb5ba0U,
+	0x52f6a452U, 0x3b4d763bU, 0xd661b7d6U, 0xb3ce7db3U, 0x297b5229U, 0xe33edde3U, 0x2f715e2fU, 0x84971384U,
+	0x53f5a653U, 0xd168b9d1U, 0x00000000U, 0xed2cc1edU, 0x20604020U, 0xfc1fe3fcU, 0xb1c879b1U, 0x5bedb65bU,
+	0x6abed46aU, 0xcb468dcbU, 0xbed967beU, 0x394b7239U, 0x4ade944aU, 0x4cd4984cU, 0x58e8b058U, 0xcf4a85cfU,
+	0xd06bbbd0U, 0xef2ac5efU, 0xaae54faaU, 0xfb16edfbU, 0x43c58643U, 0x4dd79a4dU, 0x33556633U, 0x85941185U,
+	0x45cf8a45U, 0xf910e9f9U, 0x02060402U, 0x7f81fe7fU, 0x50f0a050U, 0x3c44783cU, 0x9fba259fU, 0xa8e34ba8U,
+	0x51f3a251U, 0xa3fe5da3U, 0x40c08040U, 0x8f8a058fU, 0x92ad3f92U, 0x9dbc219dU, 0x38487038U, 0xf504f1f5U,
+	0xbcdf63bcU, 0xb6c177b6U, 0xda75afdaU, 0x21634221U, 0x10302010U, 0xff1ae5ffU, 0xf30efdf3U, 0xd26dbfd2U,
+	0xcd4c81cdU, 0x0c14180cU, 0x13352613U, 0xec2fc3ecU, 0x5fe1be5fU, 0x97a23597U, 0x44cc8844U, 0x17392e17U,
+	0xc45793c4U, 0xa7f255a7U, 0x7e82fc7eU, 0x3d477a3dU, 0x64acc864U, 0x5de7ba5dU, 0x192b3219U, 0x7395e673U,
+	0x60a0c060U, 0x81981981U, 0x4fd19e4fU, 0xdc7fa3dcU, 0x22664422U, 0x2a7e542aU, 0x90ab3b90U, 0x88830b88U,
+	0x46ca8c46U, 0xee29c7eeU, 0xb8d36bb8U, 0x143c2814U, 0xde79a7deU, 0x5ee2bc5eU, 0x0b1d160bU, 0xdb76addbU,
+	0xe03bdbe0U, 0x32566432U, 0x3a4e743aU, 0x0a1e140aU, 0x49db9249U, 0x060a0c06U, 0x246c4824U, 0x5ce4b85cU,
+	0xc25d9fc2U, 0xd36ebdd3U, 0xacef43acU, 0x62a6c462U, 0x91a83991U, 0x95a43195U, 0xe437d3e4U, 0x798bf279U,
+	0xe732d5e7U, 0xc8438bc8U, 0x37596e37U, 0x6db7da6dU, 0x8d8c018dU, 0xd564b1d5U, 0x4ed29c4eU, 0xa9e049a9U,
+	0x6cb4d86cU, 0x56faac56U, 0xf407f3f4U, 0xea25cfeaU, 0x65afca65U, 0x7a8ef47aU, 0xaee947aeU, 0x08181008U,
+	0xbad56fbaU, 0x7888f078U, 0x256f4a25U, 0x2e725c2eU, 0x1c24381cU, 0xa6f157a6U, 0xb4c773b4U, 0xc65197c6U,
+	0xe823cbe8U, 0xdd7ca1ddU, 0x749ce874U, 0x1f213e1fU, 0x4bdd964bU, 0xbddc61bdU, 0x8b860d8bU, 0x8a850f8aU,
+	0x7090e070U, 0x3e427c3eU, 0xb5c471b5U, 0x66aacc66U, 0x48d89048U, 0x03050603U, 0xf601f7f6U, 0x0e121c0eU,
+	0x61a3c261U, 0x355f6a35U, 0x57f9ae57U, 0xb9d069b9U, 0x86911786U, 0xc15899c1U, 0x1d273a1dU, 0x9eb9279eU,
+	0xe138d9e1U, 0xf813ebf8U, 0x98b32b98U, 0x11332211U, 0x69bbd269U, 0xd970a9d9U, 0x8e89078eU, 0x94a73394U,
+	0x9bb62d9bU, 0x1e223c1eU, 0x87921587U, 0xe920c9e9U, 0xce4987ceU, 0x55ffaa55U, 0x28785028U, 0xdf7aa5dfU,
+	0x8c8f038cU, 0xa1f859a1U, 0x89800989U, 0x0d171a0dU, 0xbfda65bfU, 0xe631d7e6U, 0x42c68442U, 0x68b8d068U,
+	0x41c38241U, 0x99b02999U, 0x2d775a2dU, 0x0f111e0fU, 0xb0cb7bb0U, 0x54fca854U, 0xbbd66dbbU, 0x163a2c16U,
+};
+const uint32 FAES128CTR::T3[256] =
+{
+	0x6363a5c6U, 0x7c7c84f8U, 0x777799eeU, 0x7b7b8df6U, 0xf2f20dffU, 0x6b6bbdd6U, 0x6f6fb1deU, 0xc5c55491U,
+	0x30305060U, 0x01010302U, 0x6767a9ceU, 0x2b2b7d56U, 0xfefe19e7U, 0xd7d762b5U, 0xababe64dU, 0x76769aecU,
+	0xcaca458fU, 0x82829d1fU, 0xc9c94089U, 0x7d7d87faU, 0xfafa15efU, 0x5959ebb2U, 0x4747c98eU, 0xf0f00bfbU,
+	0xadadec41U, 0xd4d467b3U, 0xa2a2fd5fU, 0xafafea45U, 0x9c9cbf23U, 0xa4a4f753U, 0x727296e4U, 0xc0c05b9bU,
+	0xb7b7c275U, 0xfdfd1ce1U, 0x9393ae3dU, 0x26266a4cU, 0x36365a6cU, 0x3f3f417eU, 0xf7f702f5U, 0xcccc4f83U,
+	0x34345c68U, 0xa5a5f451U, 0xe5e534d1U, 0xf1f108f9U, 0x717193e2U, 0xd8d873abU, 0x31315362U, 0x15153f2aU,
+	0x04040c08U, 0xc7c75295U, 0x23236546U, 0xc3c35e9dU, 0x18182830U, 0x9696a137U, 0x05050f0aU, 0x9a9ab52fU,
+	0x0707090eU, 0x12123624U, 0x80809b1bU, 0xe2e23ddfU, 0xebeb26cdU, 0x2727694eU, 0xb2b2cd7fU, 0x75759feaU,
+	0x09091b12U, 0x83839e1dU, 0x2c2c7458U, 0x1a1a2e34U, 0x1b1b2d36U, 0x6e6eb2dcU, 0x5a5aeeb4U, 0xa0a0fb5bU,
+	0x5252f6a4U, 0x3b3b4d76U, 0xd6d661b7U, 0xb3b3ce7dU, 0x29297b52U, 0xe3e33eddU, 0x2f2f715eU, 0x84849713U,
+	0x5353f5a6U, 0xd1d168b9U, 0x00000000U, 0xeded2cc1U, 0x20206040U, 0xfcfc1fe3U, 0xb1b1c879U, 0x5b5bedb6U,
+	0x6a6abed4U, 0xcbcb468dU, 0xbebed967U, 0x39394b72U, 0x4a4ade94U, 0x4c4cd498U, 0x5858e8b0U, 0xcfcf4a85U,
+	0xd0d06bbbU, 0xefef2ac5U, 0xaaaae54fU, 0xfbfb16edU, 0x4343c586U, 0x4d4dd79aU, 0x33335566U, 0x85859411U,
+	0x4545cf8aU, 0xf9f910e9U, 0x02020604U, 0x7f7f81feU, 0x5050f0a0U, 0x3c3c4478U, 0x9f9fba25U, 0xa8a8e34bU,
+	0x5151f3a2U, 0xa3a3fe5dU, 0x4040c080U, 0x8f8f8a05U, 0x9292ad3fU, 0x9d9dbc21U, 0x38384870U, 0xf5f504f1U,
+	0xbcbcdf63U, 0xb6b6c177U, 0xdada75afU, 0x21216342U, 0x10103020U, 0xffff1ae5U, 0xf3f30efdU, 0xd2d26dbfU,
+	0xcdcd4c81U, 0x0c0c1418U, 0x13133526U, 0xecec2fc3U, 0x5f5fe1beU, 0x9797a235U, 0x4444cc88U, 0x1717392eU,
+	0xc4c45793U, 0xa7a7f255U, 0x7e7e82fcU, 0x3d3d477aU, 0x6464acc8U, 0x5d5de7baU, 0x19192b32U, 0x737395e6U,
+	0x6060a0c0U, 0x81819819U, 0x4f4fd19eU, 0xdcdc7fa3U, 0x22226644U, 0x2a2a7e54U, 0x9090ab3bU, 0x8888830bU,
+	0x4646ca8cU, 0xeeee29c7U, 0xb8b8d36bU, 0x14143c28U, 0xdede79a7U, 0x5e5ee2bcU, 0x0b0b1d16U, 0xdbdb76adU,
+	0xe0e03bdbU, 0x32325664U, 0x3a3a4e74U, 0x0a0a1e14U, 0x4949db92U, 0x06060a0cU, 0x24246c48U, 0x5c5ce4b8U,
+	0xc2c25d9fU, 0xd3d36ebdU, 0xacacef43U, 0x6262a6c4U, 0x9191a839U, 0x9595a431U, 0xe4e437d3U, 0x79798bf2U,
+	0xe7e732d5U, 0xc8c8438bU, 0x3737596eU, 0x6d6db7daU, 0x8d8d8c01U, 0xd5d564b1U, 0x4e4ed29cU, 0xa9a9e049U,
+	0x6c6cb4d8U, 0x5656faacU, 0xf4f407f3U, 0xeaea25cfU, 0x6565afcaU, 0x7a7a8ef4U, 0xaeaee947U, 0x08081810U,
+	0xbabad56fU, 0x787888f0U, 0x25256f4aU, 0x2e2e725cU, 0x1c1c2438U, 0xa6a6f157U, 0xb4b4c773U, 0xc6c65197U,
+	0xe8e823cbU, 0xdddd7ca1U, 0x74749ce8U, 0x1f1f213eU, 0x4b4bdd96U, 0xbdbddc61U, 0x8b8b860dU, 0x8a8a850fU,
+	0x707090e0U, 0x3e3e427cU, 0xb5b5c471U, 0x6666aaccU, 0x4848d890U, 0x03030506U, 0xf6f601f7U, 0x0e0e121cU,
+	0x6161a3c2U, 0x35355f6aU, 0x5757f9aeU, 0xb9b9d069U, 0x86869117U, 0xc1c15899U, 0x1d1d273aU, 0x9e9eb927U,
+	0xe1e138d9U, 0xf8f813ebU, 0x9898b32bU, 0x11113322U, 0x6969bbd2U, 0xd9d970a9U, 0x8e8e8907U, 0x9494a733U,
+	0x9b9bb62dU, 0x1e1e223cU, 0x87879215U, 0xe9e920c9U, 0xcece4987U, 0x5555ffaaU, 0x28287850U, 0xdfdf7aa5U,
+	0x8c8c8f03U, 0xa1a1f859U, 0x89898009U, 0x0d0d171aU, 0xbfbfda65U, 0xe6e631d7U, 0x4242c684U, 0x6868b8d0U,
+	0x4141c382U, 0x9999b029U, 0x2d2d775aU, 0x0f0f111eU, 0xb0b0cb7bU, 0x5454fca8U, 0xbbbbd66dU, 0x16163a2cU,
+};
+const uint32 FAES128CTR::RoundConstant[10] =
+{
+	0x01000000, 0x02000000, 0x04000000, 0x08000000, 0x10000000, 0x20000000, 0x40000000, 0x80000000, 0x1B000000, 0x36000000
+};
+
+/***************************************************************************************************************************************************/
+/***************************************************************************************************************************************************/
+/***************************************************************************************************************************************************/
+
+class FStreamDecrypterAES128 : public IStreamDecrypterAES128
+{
+public:
+	FStreamDecrypterAES128();
+	virtual ~FStreamDecrypterAES128();
+
+	virtual EResult CBCInit(const TArray<uint8>& Key, const TArray<uint8>* OptionalIV=nullptr) override;
+	virtual EResult CBCDecryptInPlace(int32& OutNumBytes, uint8* InOutData, int32 NumBytes16, bool bIsFinalBlock) override;
+	virtual int32 CBCGetEncryptionDataSize(int32 PlaintextSize) override;
+	virtual EResult CBCEncryptInPlace(int32& OutNumBytes, uint8* InOutData, int32 NumBytes, bool bIsFinalData) override;
+
+	virtual EResult CTRInit(const TArray<uint8>& Key) override;
+	virtual EResult CTRSetKey(const TArray<uint8>& Key) override;
+	virtual EResult CTRSetIV(const TArray<uint8>& IV) override;
+	virtual EResult CTRDecryptInPlace(uint8* InOutData, int32 InNumBytes) override;
+private:
+	TinyAES128::AES_ctx		AES128CBCContext;
+	FAES128CTR				AES128CTR;
+	TArray<uint8>			CryptoKey;
+	bool					bIsInitialized;
+};
+
+/***************************************************************************************************************************************************/
+/***************************************************************************************************************************************************/
+/***************************************************************************************************************************************************/
+
+TSharedPtr<IStreamDecrypterAES128, ESPMode::ThreadSafe> IStreamDecrypterAES128::Create()
+{
+	return TSharedPtr<IStreamDecrypterAES128, ESPMode::ThreadSafe>(new FStreamDecrypterAES128);
+}
+
+/***************************************************************************************************************************************************/
+/***************************************************************************************************************************************************/
+/***************************************************************************************************************************************************/
+
+const TCHAR* IStreamDecrypterAES128::GetResultText(EResult ResultCode)
+{
+	switch(ResultCode)
+	{
+		case EResult::Ok:				return TEXT("Ok");
+		case EResult::NotInitialized:	return TEXT("Not initialized");
+		case EResult::BadKeyLength:		return TEXT("Bad key length");
+		case EResult::BadIVLength:		return TEXT("Bad IV length");
+		case EResult::BadDataLength:	return TEXT("Bad data length");
+		case EResult::InvalidArg:		return TEXT("Invalid argument");
+		case EResult::BadHexChar:		return TEXT("Invalid hex char");
+		default:						return TEXT("?");
+	}
+}
+
+
+IStreamDecrypterAES128::EResult IStreamDecrypterAES128::ConvHexStringToBin(TArray<uint8>& OutBinData, const char* InHexString)
+{
+	if (InHexString)
+	{
+		int32 InLen = FCStringAnsi::Strlen(InHexString);
+		// Length must be even since we are only converting full bytes.
+		if ((InLen & 1) == 0)
+		{
+			OutBinData.Init(0, InLen / 2);
+			uint8* RawData = OutBinData.GetData() + (InLen/2 - 1);
+			int32 hi = 0;
+			while(--InLen >= 0)
+			{
+				uint8 v = InHexString[InLen];
+				if (v >= '0' && v <= '9')
+				{
+					v -= '0';
+				}
+				else if (v >= 'a' && v <= 'f')
+				{
+					v = v - 'a' + 10;
+				}
+				else if (v >= 'A' && v <= 'F')
+				{
+					v = v - 'A' + 10;
+				}
+				else
+				{
+					return IStreamDecrypterAES128::EResult::BadHexChar;
+				}
+				if (hi)
+				{
+					*RawData |= (v << 4);
+					--RawData;
+				}
+				else
+				{
+					*RawData |= v;
+				}
+				hi ^= 1;
+			}
+			return IStreamDecrypterAES128::EResult::Ok;
+		}
+
+		return IStreamDecrypterAES128::EResult::BadDataLength;
+	}
+	return IStreamDecrypterAES128::EResult::InvalidArg;
+}
+
+void IStreamDecrypterAES128::MakePaddedIVFromUInt64(TArray<uint8>& OutBinData, uint64 lower64Bits)
+{
+	OutBinData.Init(0, 16);
+	for(int32 i=7; i>=0; --i)
+	{
+		OutBinData[8 + i] = (uint8)(lower64Bits & 0xff);
+		lower64Bits >>= 8;
+	}
+}
+
+/***************************************************************************************************************************************************/
+/***************************************************************************************************************************************************/
+/***************************************************************************************************************************************************/
+
+FStreamDecrypterAES128::FStreamDecrypterAES128()
+{
+	bIsInitialized = false;
+}
+
+FStreamDecrypterAES128::~FStreamDecrypterAES128()
+{
+}
+
+
+IStreamDecrypterAES128::EResult FStreamDecrypterAES128::CBCInit(const TArray<uint8>& Key, const TArray<uint8>* OptionalIV)
+{
+	bIsInitialized = false;
+	if (Key.Num() == 16)
+	{
+		if (OptionalIV && OptionalIV->Num() != 16)
+		{
+			return IStreamDecrypterAES128::EResult::BadIVLength;
+		}
+		if (OptionalIV)
+		{
+			TinyAES128::AES_init_ctx_iv(&AES128CBCContext, (const uint8_t*)Key.GetData(), (const uint8_t*)OptionalIV->GetData());
+		}
+		else
+		{
+			TinyAES128::AES_init_ctx(&AES128CBCContext, (const uint8_t*)Key.GetData());
+		}
+		bIsInitialized = true;
+		CryptoKey = Key;
+		return IStreamDecrypterAES128::EResult::Ok;
+	}
+	return IStreamDecrypterAES128::EResult::BadKeyLength;
+}
+
+IStreamDecrypterAES128::EResult FStreamDecrypterAES128::CBCDecryptInPlace(int32& OutNumBytes, uint8* InOutData, int32 NumBytes16, bool bIsFinalBlock)
+{
+	OutNumBytes = 0;
+	if (bIsInitialized)
+	{
+		// Because encryption adds PCKS7 padding the number of encrypted bytes is always a multiple of 16.
+		if ((NumBytes16 & 15) == 0)
+		{
+			if (NumBytes16)
+			{
+				TinyAES128::AES_CBC_decrypt_buffer(&AES128CBCContext, (uint8_t*)InOutData, (uint32_t)NumBytes16);
+				OutNumBytes = NumBytes16;
+				if (bIsFinalBlock)
+				{
+					// Get the last byte of the last block which contains the PCKS7 number of padding bytes added.
+					// The actual number of decrypted bytes is less that amount.
+					OutNumBytes -= (int32)InOutData[NumBytes16 - 1];
+				}
+			}
+			return IStreamDecrypterAES128::EResult::Ok;
+		}
+		return IStreamDecrypterAES128::EResult::BadDataLength;
+	}
+	return IStreamDecrypterAES128::EResult::NotInitialized;
+}
+
+int32 FStreamDecrypterAES128::CBCGetEncryptionDataSize(int32 PlaintextSize)
+{
+	// PKCS7 padding adds as many bytes as are required to get to the next 16 byte multiple.
+	// If already a 16 byte multiple a whole 16 byte pad block is added, so encrypting 0 bytes results in an output of 16 bytes!
+	int32 NumRemaining = PlaintextSize & 15;
+	return NumRemaining ? ((PlaintextSize + 15) & ~15) : PlaintextSize + 16;
+}
+
+IStreamDecrypterAES128::EResult FStreamDecrypterAES128::CBCEncryptInPlace(int32& OutNumBytes, uint8* InOutData, int32 NumBytes, bool bIsFinalData)
+{
+	OutNumBytes = 0;
+	if (bIsInitialized)
+	{
+		int32 NumRemaining = NumBytes & 15;
+		int32 NumBlockBytes = NumBytes & ~15;
+		// We cannot encrypt a non-16 byte multiple on any but the final block.
+		if (NumRemaining && !bIsFinalData)
+		{
+			return IStreamDecrypterAES128::EResult::BadDataLength;
+		}
+		if (NumBlockBytes)
+		{
+			TinyAES128::AES_CBC_encrypt_buffer(&AES128CBCContext, (uint8_t*)InOutData, (uint32_t)NumBlockBytes);
+			OutNumBytes = NumBlockBytes;
+		}
+		// The input to be encrypted must be overallocated to accommodate the extra bytes due to the final PCKS7 padding block.
+		// The size of the final encrypted data can be obtained through CBCGetExtraEncryptionDataSize().
+		if (bIsFinalData)
+		{
+			// Append the padding value
+			FMemory::Memset(InOutData + NumBytes, (uint8)(16 - NumRemaining), 16 - NumRemaining);
+			TinyAES128::AES_CBC_encrypt_buffer(&AES128CBCContext, (uint8_t*)(InOutData + NumBlockBytes), 16);
+			OutNumBytes += 16;
+		}
+		return IStreamDecrypterAES128::EResult::Ok;
+	}
+	return IStreamDecrypterAES128::EResult::NotInitialized;
+}
+
+
+
+IStreamDecrypterAES128::EResult FStreamDecrypterAES128::CTRInit(const TArray<uint8>& Key)
+{
+	bIsInitialized = false;
+	if (Key.Num() == 16)
+	{
+		AES128CTR.SetKey(Key.GetData());
+		bIsInitialized = true;
+		CryptoKey = Key;
+		return IStreamDecrypterAES128::EResult::Ok;
+	}
+	return IStreamDecrypterAES128::EResult::BadKeyLength;
+}
+
+IStreamDecrypterAES128::EResult FStreamDecrypterAES128::CTRSetKey(const TArray<uint8>& Key)
+{
+	if (bIsInitialized)
+	{
+		if (Key.Num() == 16)
+		{
+			AES128CTR.SetKey(Key.GetData());
+			CryptoKey = Key;
+			return IStreamDecrypterAES128::EResult::Ok;
+		}
+		return IStreamDecrypterAES128::EResult::BadKeyLength;
+	}
+	return IStreamDecrypterAES128::EResult::NotInitialized;
+}
+
+IStreamDecrypterAES128::EResult FStreamDecrypterAES128::CTRSetIV(const TArray<uint8>& IV)
+{
+	if (bIsInitialized)
+	{
+		if (IV.Num() == 8 || IV.Num() == 16)
+		{
+			AES128CTR.SetIV(IV);
+			return IStreamDecrypterAES128::EResult::Ok;
+		}
+		return IStreamDecrypterAES128::EResult::BadIVLength;
+	}
+
+	return IStreamDecrypterAES128::EResult::NotInitialized;
+}
+
+IStreamDecrypterAES128::EResult FStreamDecrypterAES128::CTRDecryptInPlace(uint8* InOutData, int32 InNumBytes)
+{
+	if (bIsInitialized)
+	{
+		AES128CTR.Decrypt(InOutData, InNumBytes);
+		return IStreamDecrypterAES128::EResult::Ok;
+	}
+	return IStreamDecrypterAES128::EResult::NotInitialized;
+}
+
+
+}
+
+
diff --git a/Engine/Plugins/Media/ElectraCDM/Source/ElectraCDM/Private/ElectraCDM.cpp b/Engine/Plugins/Media/ElectraCDM/Source/ElectraCDM/Private/ElectraCDM.cpp
new file mode 100644
--- /dev/null
+++ b/Engine/Plugins/Media/ElectraCDM/Source/ElectraCDM/Private/ElectraCDM.cpp
@@ -0,0 +1,215 @@
+// Copyright Epic Games, Inc. All Rights Reserved.
+
+#include "ElectraCDM.h"
+//#include "ClearKey/ClearKeyCDM.h"
+#include "HAL/CriticalSection.h"
+#include "Misc/ScopeLock.h"
+
+namespace ElectraCDM
+{
+
+struct IMediaCDM::IPlayerSession
+{
+};
+
+class FMediaDRM : public IMediaCDM
+{
+public:
+	static FMediaDRM& Get()
+	{
+		static FMediaDRM This;
+		return This;
+	}
+
+	FMediaDRM();
+	virtual ~FMediaDRM();
+
+	virtual void RegisterCDM(TWeakPtr<IMediaCDMSystem, ESPMode::ThreadSafe> InCDMSystem) override;
+	virtual void GetCDMCustomJSONPrefixes(const FString& InCDMScheme, const FString& InValue, FString& OutAttributePrefix, FString& OutTextPropertyName, bool& bOutNoNamespaces) override;
+	virtual FString GetLastErrorMessage() override;
+	virtual TSharedPtr<IMediaCDMCapabilities, ESPMode::ThreadSafe> GetCDMCapabilitiesForScheme(const FString& InCDMScheme, const FString& InValue, const FString& InAdditionalElements) override;
+
+	virtual IPlayerSession* CreatePlayerSessionID() override
+	{ return new IPlayerSession; }
+	virtual void ReleasePlayerSessionID(IPlayerSession* PlayerSession) override
+	{ delete PlayerSession; }
+
+	virtual ECDMError CreateDRMClient(TSharedPtr<IMediaCDMClient, ESPMode::ThreadSafe>& OutClient, IPlayerSession* InForPlayerSession, const TArray<FCDMCandidate>& InCandidates) override;
+	virtual ECDMError ReleasePlayerSessionKeys(IPlayerSession* PlayerSession) override;
+private:
+	void RemoveUnusedSystems();
+
+	bool CompareSchemeIds(const FString& schemeA, const FString& schemeB)
+	{
+		if (schemeA.Equals(schemeB, ESearchCase::IgnoreCase))
+		{
+			return true;
+		}
+		FString sa(schemeA);
+		FString sb(schemeB);
+		if (sa.StartsWith(TEXT("urn:uuid:")))
+		{
+			sa.RightChopInline(9);
+		}
+		if (sb.StartsWith(TEXT("urn:uuid:")))
+		{
+			sb.RightChopInline(9);
+		}
+		sa.ReplaceInline(TEXT("-"), TEXT(""), ESearchCase::CaseSensitive);
+		sb.ReplaceInline(TEXT("-"), TEXT(""), ESearchCase::CaseSensitive);
+		return sa.Equals(sb, ESearchCase::IgnoreCase);
+	}
+
+	FCriticalSection Lock;
+	TArray<TWeakPtr<IMediaCDMSystem, ESPMode::ThreadSafe>> CDMSystems;
+	FString LastErrorMessage;
+};
+
+FMediaDRM::FMediaDRM()
+{
+}
+
+FMediaDRM::~FMediaDRM()
+{
+}
+
+void FMediaDRM::RemoveUnusedSystems()
+{
+	FScopeLock lock(&Lock);
+	for(int32 i=0; i<CDMSystems.Num(); ++i)
+	{
+		if (!CDMSystems[i].IsValid())
+		{
+			CDMSystems.RemoveAt(i);
+			--i;
+		}
+	}
+}
+
+
+void FMediaDRM::RegisterCDM(TWeakPtr<IMediaCDMSystem, ESPMode::ThreadSafe> InCDMSystem)
+{
+	RemoveUnusedSystems();
+	FScopeLock lock(&Lock);
+	CDMSystems.Emplace(InCDMSystem);
+}
+
+
+void FMediaDRM::GetCDMCustomJSONPrefixes(const FString& InCDMScheme, const FString& InValue, FString& OutAttributePrefix, FString& OutTextPropertyName, bool& bOutNoNamespaces)
+{
+	OutAttributePrefix.Empty();
+	OutTextPropertyName.Empty();
+	bOutNoNamespaces = false;
+	RemoveUnusedSystems();
+	TSharedPtr<IMediaCDMCapabilities, ESPMode::ThreadSafe> Capabilities;
+	FScopeLock lock(&Lock);
+	for(int32 i=0; i<CDMSystems.Num(); ++i)
+	{
+		TSharedPtr<IMediaCDMSystem, ESPMode::ThreadSafe> CDMSystem = CDMSystems[i].Pin();
+		if (CDMSystem.IsValid() && CompareSchemeIds(InCDMScheme, CDMSystem->GetSchemeID()))
+		{
+			CDMSystem->GetCDMCustomJSONPrefixes(OutAttributePrefix, OutTextPropertyName, bOutNoNamespaces);
+			break;
+		}
+	}
+}
+
+
+TSharedPtr<IMediaCDMCapabilities, ESPMode::ThreadSafe> FMediaDRM::GetCDMCapabilitiesForScheme(const FString& InCDMScheme, const FString& InValue, const FString& InAdditionalElements)
+{
+	RemoveUnusedSystems();
+	TSharedPtr<IMediaCDMCapabilities, ESPMode::ThreadSafe> Capabilities;
+	FScopeLock lock(&Lock);
+	for(int32 i=0; i<CDMSystems.Num(); ++i)
+	{
+		TSharedPtr<IMediaCDMSystem, ESPMode::ThreadSafe> CDMSystem = CDMSystems[i].Pin();
+		if (CDMSystem.IsValid() && CompareSchemeIds(InCDMScheme, CDMSystem->GetSchemeID()))
+		{
+			Capabilities = CDMSystem->GetCDMCapabilities(InValue, InAdditionalElements);
+			if (Capabilities.IsValid())
+			{
+				break;
+			}
+		}
+	}
+	return Capabilities;
+}
+
+
+FString FMediaDRM::GetLastErrorMessage()
+{
+	FScopeLock lock(&Lock);
+	return LastErrorMessage;
+}
+
+
+
+ECDMError FMediaDRM::CreateDRMClient(TSharedPtr<IMediaCDMClient, ESPMode::ThreadSafe>& OutClient, IPlayerSession* InForPlayerSession, const TArray<FCDMCandidate>& InCandidates)
+{
+	RemoveUnusedSystems();
+
+	LastErrorMessage.Empty();
+
+	// For now we narrow the candidates down to one based on the schemeID only.
+	if (InCandidates.Num() == 0)
+	{
+		LastErrorMessage = TEXT("No candidates passed");
+		return ECDMError::Failure;
+	}
+	// And really, right now we pick whichever is the first one...
+	TSharedPtr<IMediaCDMSystem, ESPMode::ThreadSafe> SelectedCDMSystem;
+	FScopeLock lock(&Lock);
+	for(int32 i=0; i<CDMSystems.Num(); ++i)
+	{
+		TSharedPtr<IMediaCDMSystem, ESPMode::ThreadSafe> CDMSystem = CDMSystems[i].Pin();
+		if (CDMSystem.IsValid() && CompareSchemeIds(CDMSystem->GetSchemeID(), InCandidates[0].SchemeId))
+		{
+			SelectedCDMSystem = CDMSystem;
+			break;
+		}
+	}
+
+	// Add all matching candidates in. This is needed for different KIDs.
+	if (SelectedCDMSystem.IsValid())
+	{
+		TArray<FCDMCandidate> CDMCandidateList;
+		for(auto &Cand : InCandidates)
+		{
+			if (CompareSchemeIds(SelectedCDMSystem->GetSchemeID(), Cand.SchemeId))
+			{
+				CDMCandidateList.Emplace(Cand);
+			}
+		}
+		ECDMError Error = SelectedCDMSystem->CreateDRMClient(OutClient, InForPlayerSession, CDMCandidateList);
+		LastErrorMessage = SelectedCDMSystem->GetLastErrorMessage();
+		return Error;
+	}
+
+	LastErrorMessage = TEXT("No matching CDM system registered");
+	return ECDMError::Failure;
+}
+
+ECDMError FMediaDRM::ReleasePlayerSessionKeys(IPlayerSession* PlayerSession)
+{
+	RemoveUnusedSystems();
+	FScopeLock lock(&Lock);
+	for(int32 i=0; i<CDMSystems.Num(); ++i)
+	{
+		TSharedPtr<IMediaCDMSystem, ESPMode::ThreadSafe> CDMSystem = CDMSystems[i].Pin();
+		CDMSystem->ReleasePlayerSessionKeys(PlayerSession);
+	}
+	return ECDMError::Success;
+}
+
+
+
+
+IMediaCDM& IMediaCDM::Get()
+{
+	return FMediaDRM::Get();
+}
+
+
+
+}
+
diff --git a/Engine/Plugins/Media/ElectraCDM/Source/ElectraCDM/Private/ElectraCDMModule.cpp b/Engine/Plugins/Media/ElectraCDM/Source/ElectraCDM/Private/ElectraCDMModule.cpp
new file mode 100644
--- /dev/null
+++ b/Engine/Plugins/Media/ElectraCDM/Source/ElectraCDM/Private/ElectraCDMModule.cpp
@@ -0,0 +1,42 @@
+// Copyright Epic Games, Inc. All Rights Reserved.
+
+#include "ElectraCDMModule.h"
+
+#include "Modules/ModuleManager.h"
+#include "IElectraCDMModule.h"
+
+#include "ElectraCDM.h"
+#include "ClearKey/ClearKeyCDM.h"
+
+
+#define LOCTEXT_NAMESPACE "ElectraCDMModule"
+
+DEFINE_LOG_CATEGORY(LogElectraCDM);
+
+// -----------------------------------------------------------------------------------------------------------------------------------
+
+class FElectraCDMModule: public IElectraCDMModule
+{
+public:
+	// IModuleInterface interface
+
+	void StartupModule() override
+	{
+		ElectraCDM::IMediaCDM& Singleton = ElectraCDM::IMediaCDM::Get();
+
+		// Register default CDMs
+		ElectraCDM::IClearKeyCDM::RegisterWith(Singleton);
+	}
+
+	void ShutdownModule() override
+	{
+	}
+
+private:
+};
+
+IMPLEMENT_MODULE(FElectraCDMModule, ElectraCDM);
+
+#undef LOCTEXT_NAMESPACE
+
+
diff --git a/Engine/Plugins/Media/ElectraCDM/Source/ElectraCDM/Private/ElectraCDMModule.h b/Engine/Plugins/Media/ElectraCDM/Source/ElectraCDM/Private/ElectraCDMModule.h
new file mode 100644
--- /dev/null
+++ b/Engine/Plugins/Media/ElectraCDM/Source/ElectraCDM/Private/ElectraCDMModule.h
@@ -0,0 +1,8 @@
+// Copyright Epic Games, Inc. All Rights Reserved.
+
+#pragma once
+
+#include <CoreMinimal.h>
+
+DECLARE_LOG_CATEGORY_EXTERN(LogElectraCDM, Log, All);
+
diff --git a/Engine/Plugins/Media/ElectraCDM/Source/ElectraCDM/Public/Crypto/StreamCryptoAES128.h b/Engine/Plugins/Media/ElectraCDM/Source/ElectraCDM/Public/Crypto/StreamCryptoAES128.h
new file mode 100644
--- /dev/null
+++ b/Engine/Plugins/Media/ElectraCDM/Source/ElectraCDM/Public/Crypto/StreamCryptoAES128.h
@@ -0,0 +1,49 @@
+// Copyright Epic Games, Inc. All Rights Reserved.
+
+#pragma once
+
+#include <CoreMinimal.h>
+
+namespace ElectraCDM
+{
+
+
+/**
+ *
+ */
+class ELECTRACDM_API IStreamDecrypterAES128
+{
+public:
+	static TSharedPtr<IStreamDecrypterAES128, ESPMode::ThreadSafe> Create();
+	virtual ~IStreamDecrypterAES128() = default;
+
+	enum class EResult
+	{
+		Ok,
+		NotInitialized,
+		BadKeyLength,
+		BadIVLength,
+		BadDataLength,
+		InvalidArg,
+		BadHexChar
+	};
+	static const TCHAR* GetResultText(EResult ResultCode);
+
+	static EResult ConvHexStringToBin(TArray<uint8>& OutBinData, const char* InHexString);
+	static void MakePaddedIVFromUInt64(TArray<uint8>& OutBinData, uint64 lower64Bits);
+
+	virtual EResult CBCInit(const TArray<uint8>& Key, const TArray<uint8>* OptionalIV=nullptr) = 0;
+	virtual EResult CBCDecryptInPlace(int32& OutNumBytes, uint8* InOutData, int32 NumBytes16, bool bIsFinalBlock) = 0;
+	virtual int32 CBCGetEncryptionDataSize(int32 PlaintextSize) = 0;
+	virtual EResult CBCEncryptInPlace(int32& OutNumBytes, uint8* InOutData, int32 NumBytes, bool bIsFinalData) = 0;
+
+	virtual EResult CTRInit(const TArray<uint8>& Key) = 0;
+	virtual EResult CTRSetKey(const TArray<uint8>& Key) = 0;
+	virtual EResult CTRSetIV(const TArray<uint8>& IV) = 0;
+	virtual EResult CTRDecryptInPlace(uint8* InOutData, int32 InNumBytes) = 0;
+};
+
+
+} // namespace ElectraCDM
+
+
diff --git a/Engine/Plugins/Media/ElectraCDM/Source/ElectraCDM/Public/ElectraCDM.h b/Engine/Plugins/Media/ElectraCDM/Source/ElectraCDM/Public/ElectraCDM.h
new file mode 100644
--- /dev/null
+++ b/Engine/Plugins/Media/ElectraCDM/Source/ElectraCDM/Public/ElectraCDM.h
@@ -0,0 +1,133 @@
+// Copyright Epic Games, Inc. All Rights Reserved.
+
+#pragma once
+
+#include <CoreMinimal.h>
+#include <Containers/Array.h>
+#include "ElectraCDMError.h"
+
+namespace ElectraCDM
+{
+
+/*
+	Note #1:
+		The mime type may (need to) be followed by a codec as per RFC 6381.
+		In addition, for video formats, the codec should be followed by resolution=WxH
+
+		Example for the MIME type:
+			video/mp4; codecs="avc1.640028"; resolution=1920x1080
+			audio/mp4; codecs="mp4a.40.2"
+
+	Note #2:
+		Parameters are as with ISO/IEC 23009-1:2019 where:
+			"InCDMScheme" refers to the ContentProtection@schemeIdUri attribute
+			"InValue" refers to the ContentProtection@value attribute
+			"InAdditionalElements" is a JSON string of all other ContentProtection element attributes and all child elements with their respective attributes.
+*/
+
+
+class IMediaCDMSystem;
+class IMediaCDMClient;
+
+class IMediaCDMCapabilities
+{
+public:
+	virtual ~IMediaCDMCapabilities() = default;
+
+	enum class ESupportResult
+	{
+		// The media type is supported.
+		Supported,
+		// The media type is not supported.
+		NotSupported,
+		// The specified MIME type needs to provide the "codecs" information
+		// to determine whether or not the type is supported.
+		CodecRequired,
+		// The specified MIME type needs to provide feature information to
+		// determine whether or not the type is supported. (eg.: resolution=WxH)
+		FeatureRequired,
+		// Failed to parse the MIME type string. Check if it is of correct syntax.
+		BadMIMEType,
+
+		// Secure decoder required.
+		SecureDecoderRequired,
+		// Secure decoder not required.
+		SecureDecoderNotRequired,
+	};
+
+	// Returns if a specified media type is supported by this CDM.
+	virtual ESupportResult SupportsType(const FString& InMimeType) = 0;
+
+	// Returns if a secure decoder is required for the specified media type.
+	// This SHOULD be able to return the proper answer WITHOUT having acquired a
+	// license since this is may be used to select/deselect streams upfront.
+	virtual ESupportResult RequiresSecureDecoder(const FString& InMimeType) = 0;
+};
+
+
+class ELECTRACDM_API IMediaCDM
+{
+public:
+	static IMediaCDM& Get();
+
+	//-------------------------------------------------------------------------
+	// CDM registration
+	//
+	virtual void RegisterCDM(TWeakPtr<IMediaCDMSystem, ESPMode::ThreadSafe> InCDMSystem) = 0;
+
+
+	//-------------------------------------------------------------------------
+	// Obtains the CDM system's preference on how attributes and text properties should be prefixed in the JSON string
+	// passed to GetCDMCapabilitiesForScheme(..., InAdditionalElements) or in the FCDMCandidate.AdditionalElements.
+	virtual void GetCDMCustomJSONPrefixes(const FString& InCDMScheme, const FString& InValue, FString& OutAttributePrefix, FString& OutTextPropertyName, bool& bOutNoNamespaces) = 0;
+
+
+	//-------------------------------------------------------------------------
+	// CDM system capabilities
+	//
+	// Returns an interface to a CDM system's capabilities. The system is identified through its unique scheme id (a UUID).
+	// If there is no CDM system for the given scheme a nullptr is returned.
+	virtual TSharedPtr<IMediaCDMCapabilities, ESPMode::ThreadSafe> GetCDMCapabilitiesForScheme(const FString& InCDMScheme, const FString& InValue, const FString& InAdditionalElements) = 0;
+
+
+	//-------------------------------------------------------------------------
+	// State
+	//
+	virtual FString GetLastErrorMessage() = 0;
+
+	
+	//-------------------------------------------------------------------------
+	// Client
+	//
+	
+	// A "player session" is a handle created by and for one instance of a customer (ie. a media player)
+	// This session will associate registered clients with the license keys they are using.
+	struct IPlayerSession;
+	virtual IPlayerSession* CreatePlayerSessionID() = 0;
+	// Releases the specified player session and its associated keys.
+	virtual void ReleasePlayerSessionID(IPlayerSession* InPlayerSessionToRelease) = 0;
+
+	struct FCDMCandidate
+	{
+		FString SchemeId;
+		FString Value;
+		FString CommonScheme;
+		FString AdditionalElements;
+		TArray<FString> DefaultKIDs;
+	};
+	// Create a single client from all the given candidates. Candidates are those whose capabilities said they are able to decrypt
+	// a particular stream. The object of this method is to settle on one CDM system (in case several claimed capability to decrypt).
+	// Returns Success or Failure if different types of streams would be using a different type of CDM.
+	virtual ECDMError CreateDRMClient(TSharedPtr<IMediaCDMClient, ESPMode::ThreadSafe>& OutClient, IPlayerSession* InForPlayerSession, const TArray<FCDMCandidate>& InCandidates) = 0;
+
+	// Releases all license keys associated by the specified session.
+	// Must be called before calling ReleasePlayerSessionID().
+	virtual ECDMError ReleasePlayerSessionKeys(IPlayerSession* InPlayerSession) = 0;
+
+protected:
+	virtual ~IMediaCDM() = default;
+};
+
+
+}
+
diff --git a/Engine/Plugins/Media/ElectraCDM/Source/ElectraCDM/Public/ElectraCDMClient.h b/Engine/Plugins/Media/ElectraCDM/Source/ElectraCDM/Public/ElectraCDMClient.h
new file mode 100644
--- /dev/null
+++ b/Engine/Plugins/Media/ElectraCDM/Source/ElectraCDM/Public/ElectraCDMClient.h
@@ -0,0 +1,157 @@
+// Copyright Epic Games, Inc. All Rights Reserved.
+
+#pragma once
+
+#include <CoreMinimal.h>
+#include "ElectraCDMError.h"
+#include "ElectraEncryptedSampleInfo.h"
+
+
+namespace ElectraCDM
+{
+
+class IMediaCDMClient;
+class IMediaCDMDecrypter;
+
+
+/**
+ * Event listener to be registered with a CDM client.
+ * This must be implemented by the application to be notified about the CDM client's needs
+ * that need to be satisfied, like performing license acquisition.
+ */
+class IMediaCDMEventListener
+{
+public:
+	virtual ~IMediaCDMEventListener() = default;
+
+	enum class ECDMEventType
+	{
+		ProvisionRequired,
+		KeyRequired,
+		KeyExpired
+	};
+
+	/**
+	 * An event triggered by a CDM system that needs to be processed by the applicate.
+	 * Please note that the InDrmClient may be an intermediate client your application did NOT create!
+	 * Do not perform checks on the client pointer.
+	 * Also do not use this client to perform any operation other than to return the requested data,
+	 * if the event is a data request.
+	 * You have to return the opaque InEventId in the reply (see SetLicenseKeyResponseData()).
+	 */
+	virtual void OnCDMEvent(ECDMEventType InEventType, TSharedPtr<ElectraCDM::IMediaCDMClient, ESPMode::ThreadSafe> InDrmClient, void* InEventId, const TArray<uint8>& InCustomData) = 0;
+};
+
+
+
+
+
+/**
+ * CDM client
+ * This interfaces the application with the CDM system.
+ * An application can use multiple clients if necessary.
+ */
+class IMediaCDMClient
+{
+public:
+	virtual ~IMediaCDMClient() = default;
+
+
+	//-------------------------------------------------------------------------
+	// Application event listener methods.
+	//
+	// Registers an application event listener to receive OnCDMEvent() messages.
+	virtual void RegisterEventListener(TWeakPtr<IMediaCDMEventListener, ESPMode::ThreadSafe> InEventListener) = 0;
+	// Unregisters a previously registered event listener.
+	virtual void UnregisterEventListener(TWeakPtr<IMediaCDMEventListener, ESPMode::ThreadSafe> InEventListener) = 0;
+
+
+	//-------------------------------------------------------------------------
+	// State
+	//
+	virtual ECDMState GetState() = 0;
+
+	virtual FString GetLastErrorMessage() = 0;
+
+	
+	//-------------------------------------------------------------------------
+	// License acquisition
+	//
+	// Prepares the client for license acquisition based on the configuration passed in IMediaCDM::CreateDRMClient()
+	virtual void PrepareLicenses() = 0;
+
+	enum EDRMClientFlags
+	{
+		EDRMFlg_None = 0,
+		EDRMFlg_AllowCustomKeyStorage = 1U << 31
+	};
+	
+	// Overrides the URL to the license server that may have been set in the initial configuration data.
+	virtual void SetLicenseServerURL(const FString& InLicenseServerURL) = 0;
+	// Returns the URL to the license server, either from configuration data or the overrride.
+	virtual void GetLicenseKeyURL(FString& OutLicenseURL) = 0;
+	// Returns the data to send to the license server as well as the HTTP method and additional HTTP headers.
+	virtual void GetLicenseKeyRequestData(TArray<uint8>& OutKeyRequestData, FString& OutHttpMethod, TArray<FString>& OutHttpHeaders, uint32& OutFlags) = 0;
+	// Sets the license server's response data. This will, if successful, also update any already existing decrypters with the keys.
+	// The EventId must be the one received in OnCDMEvent().
+	virtual ECDMError SetLicenseKeyResponseData(void* InEventId, int32 HttpResponseCode, const TArray<uint8>& InKeyResponseData) = 0;
+
+
+	//-------------------------------------------------------------------------
+	// Decrypter
+	//
+	// Creates a decrypter for the specified MIME type.
+	virtual ECDMError CreateDecrypter(TSharedPtr<IMediaCDMDecrypter, ESPMode::ThreadSafe>& OutDecrypter, const FString& InMimeType) = 0;
+};
+
+
+
+/**
+ * Decrypter instance.
+ */
+class IMediaCDMDecrypter
+{
+public:
+	virtual ~IMediaCDMDecrypter() = default;
+
+	//-------------------------------------------------------------------------
+	// State
+	//
+	virtual ECDMState GetState() = 0;
+	virtual FString GetLastErrorMessage() = 0;
+
+
+	//-------------------------------------------------------------------------
+	// License update
+	//
+	// Update from PSSH boxes encountered in the stream.
+	virtual ECDMError UpdateInitDataFromPSSH(const TArray<uint8>& InPSSHData) = 0;
+	virtual ECDMError UpdateInitDataFromMultiplePSSH(const TArray<TArray<uint8>>& InPSSHData) = 0;
+
+	// Update from a URL and additional scheme specific elements.
+	virtual ECDMError UpdateFromURL(const FString& InURL, const FString& InAdditionalElements) = 0;
+	
+	
+	//-------------------------------------------------------------------------
+	// Decryption
+	//
+
+	// Returns if this decrypter needs to work in block streaming mode (based on the encryption settings).
+	virtual bool IsBlockStreamDecrypter() = 0;
+
+	// Reinitializes the decrypter to its starting state. This should not normally be called unless an error occurred.
+	virtual void Reinitialize() = 0;
+
+	// Decrypt a non-block oriented cipher (eg. AES 128 CTR) in place.
+	virtual ECDMError DecryptInPlace(uint8* InOutData, int32 InNumDataBytes, const FMediaCDMSampleInfo& InSampleInfo) = 0;
+
+	// Decrypt a streaming cipher that requires input of fixed cipher block lengths and possibly padding at the end (eg. AES 128 CBC with PKCS#7 padding).
+	// This needs a running state during decryption that needs to be created and released.
+	struct IStreamDecryptHandle;
+	virtual ECDMError BlockStreamDecryptStart(IStreamDecryptHandle*& OutStreamDecryptContext) = 0;
+	virtual ECDMError BlockStreamDecryptInPlace(IStreamDecryptHandle* InOutStreamDecryptContext, int32& OutNumBytesDecrypted, uint8* InOutData, int32 InNumDataBytes, const FMediaCDMSampleInfo& InSampleInfo, bool bIsLastBlock) = 0;
+	virtual ECDMError BlockStreamDecryptEnd(IStreamDecryptHandle* InStreamDecryptContext) = 0;
+};
+
+}
+
diff --git a/Engine/Plugins/Media/ElectraCDM/Source/ElectraCDM/Public/ElectraCDMError.h b/Engine/Plugins/Media/ElectraCDM/Source/ElectraCDM/Public/ElectraCDMError.h
new file mode 100644
--- /dev/null
+++ b/Engine/Plugins/Media/ElectraCDM/Source/ElectraCDM/Public/ElectraCDMError.h
@@ -0,0 +1,32 @@
+// Copyright Epic Games, Inc. All Rights Reserved.
+
+#pragma once
+
+#include <CoreMinimal.h>
+
+namespace ElectraCDM
+{
+
+
+enum class ECDMError
+{
+	Success,
+	Failure,
+	NotSupported,
+	CipherModeMismatch
+};
+
+
+enum class ECDMState
+{
+	Idle,
+	WaitingForKey,
+	Ready,
+	InvalidKey,
+	KeyExpired,
+	ConfigurationError
+};
+
+
+}
+
diff --git a/Engine/Plugins/Media/ElectraCDM/Source/ElectraCDM/Public/ElectraCDMSystem.h b/Engine/Plugins/Media/ElectraCDM/Source/ElectraCDM/Public/ElectraCDMSystem.h
new file mode 100644
--- /dev/null
+++ b/Engine/Plugins/Media/ElectraCDM/Source/ElectraCDM/Public/ElectraCDMSystem.h
@@ -0,0 +1,30 @@
+// Copyright Epic Games, Inc. All Rights Reserved.
+
+#pragma once
+
+#include <CoreMinimal.h>
+#include <Containers/Array.h>
+
+namespace ElectraCDM
+{
+
+class IMediaCDMSystem
+{
+public:
+	virtual FString GetLastErrorMessage() = 0;
+
+	virtual const FString& GetSchemeID() = 0;
+
+	virtual void GetCDMCustomJSONPrefixes(FString& OutAttributePrefix, FString& OutTextPropertyName, bool& bOutNoNamespaces) = 0;
+
+	virtual TSharedPtr<IMediaCDMCapabilities, ESPMode::ThreadSafe> GetCDMCapabilities(const FString& InValue, const FString& InAdditionalElements) = 0;
+
+	virtual ECDMError CreateDRMClient(TSharedPtr<IMediaCDMClient, ESPMode::ThreadSafe>& OutClient, IMediaCDM::IPlayerSession* InForPlayerSession, const TArray<IMediaCDM::FCDMCandidate>& InCandidates) = 0;
+	virtual ECDMError ReleasePlayerSessionKeys(IMediaCDM::IPlayerSession* PlayerSession) = 0;
+
+protected:
+	virtual ~IMediaCDMSystem() = default;
+};
+
+}
+
diff --git a/Engine/Plugins/Media/ElectraCDM/Source/ElectraCDM/Public/ElectraEncryptedSampleInfo.h b/Engine/Plugins/Media/ElectraCDM/Source/ElectraCDM/Public/ElectraEncryptedSampleInfo.h
new file mode 100644
--- /dev/null
+++ b/Engine/Plugins/Media/ElectraCDM/Source/ElectraCDM/Public/ElectraEncryptedSampleInfo.h
@@ -0,0 +1,35 @@
+// Copyright Epic Games, Inc. All Rights Reserved.
+
+#pragma once
+
+#include <CoreMinimal.h>
+#include <Containers/Array.h>
+
+namespace ElectraCDM
+{
+
+
+struct FMediaCDMSampleInfo
+{
+	struct FSubSample
+	{
+		uint16 NumClearBytes = 0;
+		uint32 NumEncryptedBytes = 0;
+	};
+	struct FCryptSkipPattern
+	{
+		uint8 PatternType = 0;
+		uint8 CryptByteBlock = 0;
+		uint8 SkipByteBlock = 0;
+	};
+	TArray<uint8> DefaultKID;
+	TArray<TArray<uint8>> KIDs;
+	TArray<uint8> IV;
+	TArray<FSubSample> SubSamples;
+	FCryptSkipPattern Pattern;
+};
+
+
+
+}
+
diff --git a/Engine/Plugins/Media/ElectraCDM/Source/ElectraCDM/Public/IElectraCDMModule.h b/Engine/Plugins/Media/ElectraCDM/Source/ElectraCDM/Public/IElectraCDMModule.h
new file mode 100644
--- /dev/null
+++ b/Engine/Plugins/Media/ElectraCDM/Source/ElectraCDM/Public/IElectraCDMModule.h
@@ -0,0 +1,20 @@
+// Copyright Epic Games, Inc. All Rights Reserved.
+
+#pragma once
+
+#include "Templates/SharedPointer.h"
+#include "Modules/ModuleInterface.h"
+
+/**
+ * Interface for the `ElectraCDM` module.
+ */
+class IElectraCDMModule
+	: public IModuleInterface
+{
+public:
+
+	virtual ~IElectraCDMModule() { }
+
+private:
+};
+
diff --git a/Engine/Plugins/Media/ElectraPlayer/ElectraPlayer.uplugin b/Engine/Plugins/Media/ElectraPlayer/ElectraPlayer.uplugin
--- a/Engine/Plugins/Media/ElectraPlayer/ElectraPlayer.uplugin
+++ b/Engine/Plugins/Media/ElectraPlayer/ElectraPlayer.uplugin
@@ -48,6 +48,10 @@
 		{
 			"Name": "ElectraUtil",
 			"Enabled": true
+		},
+		{
+			"Name": "ElectraCDM",
+			"Enabled": true
 		}
 	]
 }
\ No newline at end of file
diff --git a/Engine/Plugins/Media/ElectraPlayer/Source/ElectraPlayerRuntime/ElectraPlayerRuntime.Build.cs b/Engine/Plugins/Media/ElectraPlayer/Source/ElectraPlayerRuntime/ElectraPlayerRuntime.Build.cs
--- a/Engine/Plugins/Media/ElectraPlayer/Source/ElectraPlayerRuntime/ElectraPlayerRuntime.Build.cs
+++ b/Engine/Plugins/Media/ElectraPlayer/Source/ElectraPlayerRuntime/ElectraPlayerRuntime.Build.cs
@@ -22,6 +22,7 @@ namespace UnrealBuildTool.Rules
 					"WebSockets",
 					"Json",
 					"ElectraBase",
+					"ElectraCDM",
 					"XmlParser"
 				});
 			if (Target.bCompileAgainstEngine)
diff --git a/Engine/Plugins/Media/ElectraPlayer/Source/ElectraPlayerRuntime/Private/ElectraPlayer.cpp b/Engine/Plugins/Media/ElectraPlayer/Source/ElectraPlayerRuntime/Private/ElectraPlayer.cpp
--- a/Engine/Plugins/Media/ElectraPlayer/Source/ElectraPlayerRuntime/Private/ElectraPlayer.cpp
+++ b/Engine/Plugins/Media/ElectraPlayer/Source/ElectraPlayerRuntime/Private/ElectraPlayer.cpp
@@ -10,7 +10,6 @@
 #include "Renderer/RendererVideo.h"
 #include "Renderer/RendererAudio.h"
 #include "VideoDecoderResourceDelegate.h"
-#include "Crypto/StreamCryptoAES128.h"
 #include "Utilities/Utilities.h"
 
 #include "Async/Async.h"
@@ -2336,12 +2335,10 @@ void FElectraPlayer::FAdaptiveStreamingPlayerResourceProvider::ProcessPendingSta
 					if (!LicenseKeyData.IsEmpty() && LicenseKeyData != InOutRequest->GetResourceURL())
 					{
 						TArray<uint8> BinKey;
-						IStreamDecrypterAES128::EResult DecrypterResult = IStreamDecrypterAES128::ConvHexStringToBin(BinKey, TCHAR_TO_UTF8(*LicenseKeyData));
-						if (DecrypterResult == IStreamDecrypterAES128::EResult::Ok)
-						{
-							TSharedPtr<TArray<uint8>, ESPMode::ThreadSafe> ResponseDataPtr = MakeShared<TArray<uint8>, ESPMode::ThreadSafe>(BinKey);
-							InOutRequest->SetPlaybackData(ResponseDataPtr);
-						}
+						BinKey.AddUninitialized(LicenseKeyData.Len());
+						BinKey.SetNum(HexToBytes(LicenseKeyData, BinKey.GetData()));
+						TSharedPtr<TArray<uint8>, ESPMode::ThreadSafe> ResponseDataPtr = MakeShared<TArray<uint8>, ESPMode::ThreadSafe>(BinKey);
+						InOutRequest->SetPlaybackData(ResponseDataPtr);
 					}
 				}
 			}
diff --git a/Engine/Plugins/Media/ElectraPlayer/Source/ElectraPlayerRuntime/Private/Runtime/Crypto/StreamCryptoAES128.cpp b/Engine/Plugins/Media/ElectraPlayer/Source/ElectraPlayerRuntime/Private/Runtime/Crypto/StreamCryptoAES128.cpp
deleted file mode 100644
diff --git a/Engine/Plugins/Media/ElectraPlayer/Source/ElectraPlayerRuntime/Private/Runtime/Crypto/StreamCryptoAES128.h b/Engine/Plugins/Media/ElectraPlayer/Source/ElectraPlayerRuntime/Private/Runtime/Crypto/StreamCryptoAES128.h
deleted file mode 100644
diff --git a/Engine/Plugins/Media/ElectraPlayer/Source/ElectraPlayerRuntime/Private/Runtime/Demuxer/ParserISO14496-12.cpp b/Engine/Plugins/Media/ElectraPlayer/Source/ElectraPlayerRuntime/Private/Runtime/Demuxer/ParserISO14496-12.cpp
--- a/Engine/Plugins/Media/ElectraPlayer/Source/ElectraPlayerRuntime/Private/Runtime/Demuxer/ParserISO14496-12.cpp
+++ b/Engine/Plugins/Media/ElectraPlayer/Source/ElectraPlayerRuntime/Private/Runtime/Demuxer/ParserISO14496-12.cpp
@@ -31,9 +31,14 @@ namespace Electra
 		template <typename T>
 		T ValueFromBigEndian(const T value)
 		{
-			return MEDIA_TO_BIG_ENDIAN(value);
+			return MEDIA_FROM_BIG_ENDIAN(value);
 		}
 
+		template <typename T>
+		T ValueToBigEndian(const T value)
+		{
+			return MEDIA_TO_BIG_ENDIAN(value);
+		}
 	}
 
 	class FMP4Box;
@@ -658,6 +663,11 @@ private:
 		static const IParserISO14496_12::FBoxType kSample_enca = MAKE_BOX_ATOM('e', 'n', 'c', 'a');
 		static const IParserISO14496_12::FBoxType kSample_encv = MAKE_BOX_ATOM('e', 'n', 'c', 'v');
 
+		/**
+		 * Sample grouping types
+		 */
+		static const IParserISO14496_12::FBoxType kGrouping_seig = MAKE_BOX_ATOM('s', 'e', 'i', 'g');
+
 		/**
 		 * Encryption scheme types
 		 */
@@ -772,6 +782,16 @@ private:
 			return Error;
 		}
 
+		UEMediaError ReadAndParseAttributes(FMP4BoxReader& Reader)
+		{
+			uint32 VersionAndFlags = 0;
+			UEMediaError Error = UEMEDIA_ERROR_OK;
+			RETURN_IF_ERROR(Reader.Read(VersionAndFlags));
+			Version = VersionAndFlags >> 24;
+			Flags = VersionAndFlags & 0x00ffffff;
+			return Error;
+		}
+
 		uint32		Flags;
 		uint8		Version;
 	};
@@ -3429,6 +3449,22 @@ private:
 		{
 		}
 
+		const TArray<uint8>& GetBoxData()  const
+		{
+			return BoxData;
+		}
+		const TArray<uint8>& GetSystemID() const
+		{
+			return SystemID;
+		}
+		const TArray<TArray<uint8>>& GetKIDs() const
+		{
+			return KIDs;
+		}
+		const TArray<uint8>& GetData() const
+		{
+			return Data;
+		}
 	private:
 		FMP4BoxPSSH() = delete;
 		FMP4BoxPSSH(const FMP4BoxPSSH&) = delete;
@@ -3437,32 +3473,46 @@ private:
 		virtual UEMediaError ReadAndParseAttributes(FMP4ParseInfo* ParseInfo) override
 		{
 			UEMediaError Error = UEMEDIA_ERROR_OK;
-			RETURN_IF_ERROR(FMP4BoxFull::ReadAndParseAttributes(ParseInfo));
+			// We need to retain the box as a whole since some DRM systems need it to be passed in verbatim.
+			// In order to do that we read it into a buffer first, which can be handed out later if needed
+			// and parse the box from the buffer.
+			uint32 BytesRemaining = BoxSize - (ParseInfo->Reader()->GetCurrentReadOffset() - StartOffset);
+			BoxData.AddUninitialized(8 + BytesRemaining);
+			RETURN_IF_ERROR(ParseInfo->Reader()->ReadBytes(BoxData.GetData() + 8, BytesRemaining));
+			uint32* RawBox = (uint32*)BoxData.GetData();
+			RawBox[0] = Utils::ValueToBigEndian((uint32)(BytesRemaining+8));
+			RawBox[1] = Utils::ValueToBigEndian(kBox_pssh);
+
+			FDataBufferReader BufferReader(BoxData);
+			FMP4BoxReader Reader(&BufferReader, nullptr);
+			Reader.ReadBytes(nullptr, 8);
+
+			RETURN_IF_ERROR(FMP4BoxFull::ReadAndParseAttributes(Reader));
 
 			SystemID.AddUninitialized(16);
-			RETURN_IF_ERROR(ParseInfo->Reader()->ReadBytes(SystemID.GetData(), 16));		// SystemID
+			RETURN_IF_ERROR(Reader.ReadBytes(SystemID.GetData(), 16));		// SystemID
 			if (Version > 0)
 			{
 				uint32 KID_count = 0;
-				RETURN_IF_ERROR(ParseInfo->Reader()->Read(KID_count));						// KID_count
+				RETURN_IF_ERROR(Reader.Read(KID_count));					// KID_count
 				for(uint32 i=0; i<KID_count; ++i)
 				{
 					TArray<uint8>& kid = KIDs.AddDefaulted_GetRef();
 					kid.AddUninitialized(16);
-					RETURN_IF_ERROR(ParseInfo->Reader()->ReadBytes(kid.GetData(), 16));		// KID
+					RETURN_IF_ERROR(Reader.ReadBytes(kid.GetData(), 16));	// KID
 				}
 			}
 			uint32 DataSize = 0;
-			RETURN_IF_ERROR(ParseInfo->Reader()->Read(DataSize));							// DataSize
+			RETURN_IF_ERROR(Reader.Read(DataSize));							// DataSize
 			if (DataSize)
 			{
 				Data.AddUninitialized(DataSize);
-				RETURN_IF_ERROR(ParseInfo->Reader()->ReadBytes(Data.GetData(), DataSize));	// Data
+				RETURN_IF_ERROR(Reader.ReadBytes(Data.GetData(), DataSize));// Data
 			}
 			return Error;
 		}
-
 	private:
+		TArray<uint8>			BoxData;
 		TArray<uint8>			SystemID;
 		TArray<TArray<uint8>>	KIDs;
 		TArray<uint8>			Data;
@@ -3767,10 +3817,10 @@ private:
 			if (EntryCount)
 			{
 				Offsets.AddUninitialized(EntryCount);
+				uint64* Offs = Offsets.GetData();
 				if (Version == 0)
 				{
 					uint32 Value32 = 0;
-					uint64* Offs = Offsets.GetData();
 					for(uint32 i=0; i<EntryCount; ++i)
 					{
 						RETURN_IF_ERROR(ParseInfo->Reader()->Read(Value32));					// offset
@@ -3779,7 +3829,12 @@ private:
 				}
 				else
 				{
-					RETURN_IF_ERROR(ParseInfo->Reader()->ReadBytes(Offsets.GetData(), EntryCount));// offset
+					uint64 Value64 = 0;
+					for(uint32 i=0; i<EntryCount; ++i)
+					{
+						RETURN_IF_ERROR(ParseInfo->Reader()->Read(Value64));					// offset
+						*Offs++ = Value64;
+					}
 				}
 			}
 			return Error;
@@ -3811,12 +3866,7 @@ private:
 		struct FEntry
 		{
 			TArray<uint8> IV;
-			struct FSubSample
-			{
-				uint16	BytesOfClearData;
-				uint32	BytesOfProtectedData;
-			};
-			TArray<FSubSample> SubSamples;
+			TArray<ElectraCDM::FMediaCDMSampleInfo::FSubSample> SubSamples;
 		};
 
 		const TArray<FEntry>& GetEntries() const
@@ -3855,9 +3905,9 @@ private:
 						RETURN_IF_ERROR(Reader.Read(SubSampleCount));
 						for(int32 j=0; j<SubSampleCount; ++j)
 						{
-							FEntry::FSubSample& SubSample = Entry.SubSamples.AddDefaulted_GetRef();
-							RETURN_IF_ERROR(Reader.Read(SubSample.BytesOfClearData));
-							RETURN_IF_ERROR(Reader.Read(SubSample.BytesOfProtectedData));
+							ElectraCDM::FMediaCDMSampleInfo::FSubSample& SubSample = Entry.SubSamples.AddDefaulted_GetRef();
+							RETURN_IF_ERROR(Reader.Read(SubSample.NumClearBytes));
+							RETURN_IF_ERROR(Reader.Read(SubSample.NumEncryptedBytes));
 						}
 					}
 				}
@@ -3901,6 +3951,181 @@ private:
 	};
 
 
+	/**
+	 * 'sbgp' box. ISO/IEC 14496-12:2015 - 8.9.2 Sample to Group Box
+	 */
+	class FMP4BoxSBGP : public FMP4BoxFull
+	{
+	public:
+		FMP4BoxSBGP(IParserISO14496_12::FBoxType InBoxType, int64 InBoxSize, int64 InStartOffset, int64 InDataOffset, bool bInIsLeafBox)
+			: FMP4BoxFull(InBoxType, InBoxSize, InStartOffset, InDataOffset, bInIsLeafBox)
+			, GroupingType(0), GroupingTypeParameter(0)
+		{
+		}
+
+		virtual ~FMP4BoxSBGP()
+		{
+		}
+
+		struct FEntry
+		{
+			uint32 SampleCount;
+			uint32 GroupDescriptionIndex;
+		};
+
+		uint32 GetGroupingType() const
+		{ return GroupingType; }
+
+		uint32 GetGroupingTypeParameter() const
+		{ return GroupingTypeParameter; }
+
+		int32 GetNumberOfEntries() const
+		{ return Entries.Num(); }
+		
+		const FEntry& GetEntry(int32 Index) const
+		{ return Entries[Index]; }
+
+	private:
+		FMP4BoxSBGP() = delete;
+		FMP4BoxSBGP(const FMP4BoxSBGP&) = delete;
+
+	protected:
+		virtual UEMediaError ReadAndParseAttributes(FMP4ParseInfo* ParseInfo) override
+		{
+			UEMediaError Error = UEMEDIA_ERROR_OK;
+			RETURN_IF_ERROR(FMP4BoxFull::ReadAndParseAttributes(ParseInfo));
+
+			RETURN_IF_ERROR(ParseInfo->Reader()->Read(GroupingType));						// grouping_type
+			if (Version == 1)
+			{
+				RETURN_IF_ERROR(ParseInfo->Reader()->Read(GroupingTypeParameter));			// grouping_type_parameter
+			}
+			uint32 EntryCount = 0;
+			RETURN_IF_ERROR(ParseInfo->Reader()->Read(EntryCount));							// entry_count
+			for(uint32 i=0; i<EntryCount; ++i)
+			{
+				FEntry& Entry = Entries.AddDefaulted_GetRef();
+				RETURN_IF_ERROR(ParseInfo->Reader()->Read(Entry.SampleCount));				// sample_count
+				RETURN_IF_ERROR(ParseInfo->Reader()->Read(Entry.GroupDescriptionIndex));	// group_description_index
+			}
+			return Error;
+		}
+
+	private:
+		uint32	GroupingType;
+		uint32	GroupingTypeParameter;
+		TArray<FEntry> Entries;
+	};
+
+
+
+	/**
+	 * Base class of a sample group description entry. ISO/IEC 14496-12:2016 - 8.9.3.2 Sample Group Description Box, Syntax
+	 */
+	struct FSampleGroupDescriptionEntry
+	{
+		TArray<uint8> Data;
+	};
+
+	/**
+	 * 'seig' group type entry. ISO/IEC 23007-1:2016 - 6 Encryption parameters shared by groups of samples
+	 */
+	struct FCencSampleEncryptionInformationGroupEntry : FSampleGroupDescriptionEntry
+	{
+		uint8 GetCryptByteBlock() const { return Data[1] >> 4; }
+		uint8 GetSkipByteBlock() const { return Data[1] & 0x0f; }
+		uint8 GetIsProtected() const { return Data[2]; }
+		uint8 GetPerSampleIVSize() const { return Data[3]; }
+		TArray<uint8> GetKID() const { return TArray<uint8>(Data.GetData() + 4, 16); }
+		uint8 GetConstantIVSize() const { return GetIsProtected() && GetPerSampleIVSize() == 0 ? Data[20] : 0; }
+		TArray<uint8> GetConstantIV() const { uint8 civSize = GetConstantIVSize(); return civSize == 0 ? TArray<uint8>() : TArray<uint8>(Data.GetData() + 21, civSize); }
+	};
+
+	/**
+	 * 'sgpd' box. ISO/IEC 14496-12:2015 - 8.9.3 Sample Group Description Box
+	 */
+	class FMP4BoxSGPD : public FMP4BoxFull
+	{
+	public:
+		FMP4BoxSGPD(IParserISO14496_12::FBoxType InBoxType, int64 InBoxSize, int64 InStartOffset, int64 InDataOffset, bool bInIsLeafBox)
+			: FMP4BoxFull(InBoxType, InBoxSize, InStartOffset, InDataOffset, bInIsLeafBox)
+			, GroupingType(0), DefaultLength(0), DefaultSampleDescriptionIndex(0)
+		{
+		}
+
+		virtual ~FMP4BoxSGPD()
+		{
+		}
+
+		uint32 GetGroupingType() const
+		{ return GroupingType; }
+
+		uint32 GetDefaultSampleDescriptionIndex() const
+		{ return DefaultSampleDescriptionIndex; }
+
+		int32 GetNumberOfEntries() const
+		{ return Entries.Num(); }
+		
+		const FSampleGroupDescriptionEntry& GetEntry(int32 Index) const
+		{ return Entries[Index]; }
+
+	private:
+		FMP4BoxSGPD() = delete;
+		FMP4BoxSGPD(const FMP4BoxSGPD&) = delete;
+
+	protected:
+		virtual UEMediaError ReadAndParseAttributes(FMP4ParseInfo* ParseInfo) override
+		{
+			UEMediaError Error = UEMEDIA_ERROR_OK;
+			RETURN_IF_ERROR(FMP4BoxFull::ReadAndParseAttributes(ParseInfo));
+
+			// Version 0 of this box type has been deprecated.
+			if (Version == 0)
+			{
+				// Skip over the contents of this box.
+				uint32 BytesRemaining = BoxSize - (ParseInfo->Reader()->GetCurrentReadOffset() - StartOffset);
+				if (BytesRemaining)
+				{
+					RETURN_IF_ERROR(ParseInfo->Reader()->ReadBytes(nullptr, BytesRemaining));
+				}
+				return Error;
+			}
+
+			RETURN_IF_ERROR(ParseInfo->Reader()->Read(GroupingType));							// grouping_type
+			
+			if (Version == 1)
+			{
+				RETURN_IF_ERROR(ParseInfo->Reader()->Read(DefaultLength));						// default_length
+			}
+			if (Version >= 2)
+			{
+				RETURN_IF_ERROR(ParseInfo->Reader()->Read(DefaultSampleDescriptionIndex));		// default_sample_description_index
+			}
+
+			uint32 EntryCount = 0;
+			RETURN_IF_ERROR(ParseInfo->Reader()->Read(EntryCount));								// entry_count
+			for(uint32 i=0; i<EntryCount; ++i)
+			{
+				uint32 EntrySize = DefaultLength;
+				if (Version == 1 && DefaultLength == 0)
+				{
+					RETURN_IF_ERROR(ParseInfo->Reader()->Read(EntrySize));						// description_length
+				}
+
+				// The actual sample group entry now follows.
+				FSampleGroupDescriptionEntry& Entry = Entries.AddDefaulted_GetRef();
+				Entry.Data.AddUninitialized(EntrySize);
+				RETURN_IF_ERROR(ParseInfo->Reader()->ReadBytes(Entry.Data.GetData(), EntrySize));
+			}
+			return Error;
+		}
+
+	private:
+		uint32 GroupingType;
+		uint32 DefaultLength;
+		uint32 DefaultSampleDescriptionIndex;
+		TArray<FSampleGroupDescriptionEntry> Entries;
+	};
 
 
 	FMP4ParseInfo::~FMP4ParseInfo()
@@ -4196,6 +4421,12 @@ private:
 						case FMP4Box::kBox_saio:
 							NextBox = new FMP4BoxSAIO(BoxType, BoxSize, BoxStartOffset, BoxDataOffset, true);
 							break;
+						case FMP4Box::kBox_sbgp:
+							NextBox = new FMP4BoxSBGP(BoxType, BoxSize, BoxStartOffset, BoxDataOffset, true);
+							break;
+						case FMP4Box::kBox_sgpd:
+							NextBox = new FMP4BoxSGPD(BoxType, BoxSize, BoxStartOffset, BoxDataOffset, true);
+							break;
 						// Boxes we ignore for now.
 						case FMP4Box::kBox_tref:
 						case FMP4Box::kBox_colr:
@@ -4429,6 +4660,7 @@ private:
 			virtual int64 GetRawPTS() const override;
 			virtual int64 GetCompositionTimeEdit() const override;
 			virtual int64 GetEmptyEditOffset() const override;
+			virtual bool GetEncryptionInfo(ElectraCDM::FMediaCDMSampleInfo& OutSampleEncryptionInfo) const override;
 
 			UEMediaError StartAtFirstInteral();
 			void SetTrack(const FTrack* InTrack);
@@ -4533,6 +4765,7 @@ private:
 			virtual const FStreamCodecInformation& GetCodecInformation() const override;
 			virtual const FBitrateInfo& GetBitrateInfo() const override;
 			virtual const FString GetLanguage() const override;
+			virtual void GetPSSHBoxes(TArray<TArray<uint8>>& OutBoxes, bool bFromMOOV, bool bFromMOOF) const override;
 
 			//private:
 			const FMP4BoxMVHD* MVHDBox = nullptr;
@@ -4550,6 +4783,7 @@ private:
 			const FMP4BoxSTCO* STCOBox = nullptr;
 			const FMP4BoxSTSS* STSSBox = nullptr;
 			// Encryption
+			TArray<const FMP4BoxPSSH*>	PSSHBoxesFromInit;
 			TArray<const FMP4BoxPSSH*>	PSSHBoxes;
 			const FMP4BoxSCHM* SCHMBox = nullptr;
 			const FMP4BoxTENC* TENCBox = nullptr;
@@ -4776,6 +5010,23 @@ private:
 		return MDHDBox ? MDHDBox->GetLanguage() : FString(TEXT("und"));
 	}
 
+	void FParserISO14496_12::FTrack::GetPSSHBoxes(TArray<TArray<uint8>>& OutBoxes, bool bFromMOOV, bool bFromMOOF) const
+	{
+		if (bFromMOOV)
+		{
+			for(int32 i=0; i<PSSHBoxesFromInit.Num(); ++i)
+			{
+				OutBoxes.Add(PSSHBoxesFromInit[i]->GetBoxData());
+			}
+		}
+		if (bFromMOOF)
+		{
+			for(int32 i=0; i<PSSHBoxes.Num(); ++i)
+			{
+				OutBoxes.Add(PSSHBoxes[i]->GetBoxData());
+			}
+		}
+	}
 
 	/***************************************************************************************************************************************************/
 	/***************************************************************************************************************************************************/
@@ -5407,6 +5658,28 @@ private:
 		return EmptyEditDurationInMediaTimeUnits;
 	}
 
+	bool FParserISO14496_12::FTrackIterator::GetEncryptionInfo(ElectraCDM::FMediaCDMSampleInfo& OutSampleEncryptionInfo) const
+	{
+		if (!Track || !Track->SENCBox)
+		{
+			return false;
+		}
+		if (Track->TENCBox)
+		{
+			OutSampleEncryptionInfo.DefaultKID = Track->TENCBox->GetDefaultKID();
+		}
+	
+		const TArray<FMP4BoxSENC::FEntry>& SencEntries = Track->SENCBox->GetEntries();
+		check((int32) SampleNumberInTRUN < SencEntries.Num());
+		if ((int32) SampleNumberInTRUN >= SencEntries.Num())
+		{
+			return false;
+		}
+		OutSampleEncryptionInfo.IV = SencEntries[SampleNumberInTRUN].IV;
+		OutSampleEncryptionInfo.SubSamples = SencEntries[SampleNumberInTRUN].SubSamples;
+		return true;
+	}
+
 
 	/***************************************************************************************************************************************************/
 	/***************************************************************************************************************************************************/
@@ -5974,12 +6247,24 @@ private:
 									Track->SAIZBox = static_cast<const FMP4BoxSAIZ*>(Box->FindBox(FMP4Box::kBox_saiz, 0));
 									Track->SAIOBox = static_cast<const FMP4BoxSAIO*>(Box->FindBox(FMP4Box::kBox_saio, 0));
 									Track->SENCBox = static_cast<const FMP4BoxSENC*>(Box->FindBox(FMP4Box::kBox_senc, 0));
-									// Get the PSSH boxes from the track fragment.
+
+									// Get the PSSH boxes from the init segment.
 									TArray<const FMP4Box*> AllPSSHBoxes;
+									if (MOOVBox)
+									{
+										MOOVBox->GetAllBoxInstances(AllPSSHBoxes, FMP4Box::kBox_pssh);
+										for(int32 nPSSHs=0; nPSSHs<AllPSSHBoxes.Num(); ++nPSSHs)
+										{
+											Track->PSSHBoxesFromInit.Add(static_cast<const FMP4BoxPSSH*>(AllPSSHBoxes[nPSSHs]));
+										}
+										AllPSSHBoxes.Empty();
+									}
+
+									// Get the PSSH boxes from the track fragment.
 									Box->GetAllBoxInstances(AllPSSHBoxes, FMP4Box::kBox_pssh);
-									for(int32 nTruns=0; nTruns<AllPSSHBoxes.Num(); ++nTruns)
+									for(int32 nPSSHs=0; nPSSHs<AllPSSHBoxes.Num(); ++nPSSHs)
 									{
-										Track->PSSHBoxes.Add(static_cast<const FMP4BoxPSSH*>(AllPSSHBoxes[nTruns]));
+										Track->PSSHBoxes.Add(static_cast<const FMP4BoxPSSH*>(AllPSSHBoxes[nPSSHs]));
 									}
 									// Now that we have the information needed to parse the senc box, do it.
 									if (Track->SENCBox)
diff --git a/Engine/Plugins/Media/ElectraPlayer/Source/ElectraPlayerRuntime/Private/Runtime/Demuxer/ParserISO14496-12.h b/Engine/Plugins/Media/ElectraPlayer/Source/ElectraPlayerRuntime/Private/Runtime/Demuxer/ParserISO14496-12.h
--- a/Engine/Plugins/Media/ElectraPlayer/Source/ElectraPlayerRuntime/Private/Runtime/Demuxer/ParserISO14496-12.h
+++ b/Engine/Plugins/Media/ElectraPlayer/Source/ElectraPlayerRuntime/Private/Runtime/Demuxer/ParserISO14496-12.h
@@ -10,6 +10,9 @@
 #include "StreamTypes.h"
 #include "ParameterDictionary.h"
 
+#include "ElectraEncryptedSampleInfo.h"
+
+
 namespace Electra
 {
 	//
@@ -191,6 +194,8 @@ namespace Electra
 			virtual int64 GetRawPTS() const = 0;
 			virtual int64 GetCompositionTimeEdit() const = 0;
 			virtual int64 GetEmptyEditOffset() const = 0;
+
+			virtual bool GetEncryptionInfo(ElectraCDM::FMediaCDMSampleInfo& OutSampleEncryptionInfo) const = 0;
 		};
 
 		class ITrack
@@ -219,6 +224,7 @@ namespace Electra
 			virtual const FStreamCodecInformation& GetCodecInformation() const = 0;
 			virtual const FBitrateInfo& GetBitrateInfo() const = 0;
 			virtual const FString GetLanguage() const = 0;
+			virtual void GetPSSHBoxes(TArray<TArray<uint8>>& OutBoxes, bool bFromMOOV, bool bFromMOOF) const = 0;
 		};
 
 
diff --git a/Engine/Plugins/Media/ElectraPlayer/Source/ElectraPlayerRuntime/Private/Runtime/HTTP/HTTPManager.cpp b/Engine/Plugins/Media/ElectraPlayer/Source/ElectraPlayerRuntime/Private/Runtime/HTTP/HTTPManager.cpp
--- a/Engine/Plugins/Media/ElectraPlayer/Source/ElectraPlayerRuntime/Private/Runtime/HTTP/HTTPManager.cpp
+++ b/Engine/Plugins/Media/ElectraPlayer/Source/ElectraPlayerRuntime/Private/Runtime/HTTP/HTTPManager.cpp
@@ -49,8 +49,8 @@ namespace Electra
 		virtual ~FElectraHttpManager();
 		void Initialize();
 
-		virtual void AddRequest(TSharedPtrTS<FRequest> Request) override;
-		virtual void RemoveRequest(TSharedPtrTS<FRequest> Request) override;
+		virtual void AddRequest(TSharedPtrTS<FRequest> Request, bool bAutoRemoveWhenComplete) override;
+		virtual void RemoveRequest(TSharedPtrTS<FRequest> Request, bool bDoNotWaitForRemoval) override;
 
 	private:
 		struct FLocalByteStream
@@ -392,27 +392,44 @@ namespace Electra
 		bThreadStarted = true;
 	}
 
-	void FElectraHttpManager::AddRequest(TSharedPtrTS<FRequest> Request)
+	void FElectraHttpManager::AddRequest(TSharedPtrTS<FRequest> Request, bool bAutoRemoveWhenComplete)
 	{
 		FScopeLock lock(&Lock);
 		if (!bTerminate)
 		{
+			// Not currently supported. Reserved for future use.
+			check(bAutoRemoveWhenComplete == false);
+			//Request->bAutoRemoveWhenComplete = bAutoRemoveWhenComplete;
 			RequestsToAdd.Enqueue(Request);
 			RequestChangesEvent.Signal();
 		}
 	}
 
-	void FElectraHttpManager::RemoveRequest(TSharedPtrTS<FRequest> Request)
+	void FElectraHttpManager::RemoveRequest(TSharedPtrTS<FRequest> Request, bool bDoNotWaitForRemoval)
 	{
-		TSharedPtrTS<FMediaEvent> WaitingEvent = MakeSharedTS<FMediaEvent>();
-		FRemoveRequest Remove;
-		Remove.Request = Request;
-		Remove.WaitingEvent = WaitingEvent;
-		Lock.Lock();
-		RequestsToRemove.Enqueue(MoveTemp(Remove));
-		RequestChangesEvent.Signal();
-		Lock.Unlock();
-		WaitingEvent->Wait();
+		if (bDoNotWaitForRemoval)
+		{
+			Request->ReceiveBuffer.Reset();
+			Request->ProgressListener.Reset();
+			FRemoveRequest Remove;
+			Remove.Request = Request;
+			Lock.Lock();
+			RequestsToRemove.Enqueue(MoveTemp(Remove));
+			RequestChangesEvent.Signal();
+			Lock.Unlock();
+		}
+		else
+		{
+			TSharedPtrTS<FMediaEvent> WaitingEvent = MakeSharedTS<FMediaEvent>();
+			FRemoveRequest Remove;
+			Remove.Request = Request;
+			Remove.WaitingEvent = WaitingEvent;
+			Lock.Lock();
+			RequestsToRemove.Enqueue(MoveTemp(Remove));
+			RequestChangesEvent.Signal();
+			Lock.Unlock();
+			WaitingEvent->Wait();
+		}
 	}
 
 	FElectraHttpManager::FHandle* FElectraHttpManager::CreateLocalFileHandle(const FTimeValue& Now, FTransportError& OutError, const TSharedPtrTS<IElectraHttpManager::FRequest>& Request)
@@ -472,11 +489,17 @@ namespace Electra
 		if (!Request->Parameters.Verb.IsEmpty())
 		{
 			Handle->HttpRequest->SetVerb(Request->Parameters.Verb);
+			// Add POST data
+			if (Request->Parameters.Verb.Equals(TEXT("POST")))
+			{
+				Handle->HttpRequest->SetContent(MoveTemp(Request->Parameters.PostData));
+			}
 		}
 		else
 		{
 			Handle->HttpRequest->SetVerb(TEXT("GET"));
 		}
+
 		Handle->HttpRequest->SetHeader(TEXT("User-Agent"), "X-UnrealEngine-Agent");
 		// Set accepted encoding first. If this is also present in custom headers let those overwrite it.
 		if (Request->Parameters.AcceptEncoding.IsSet())
diff --git a/Engine/Plugins/Media/ElectraPlayer/Source/ElectraPlayerRuntime/Private/Runtime/HTTP/HTTPManager.h b/Engine/Plugins/Media/ElectraPlayer/Source/ElectraPlayerRuntime/Private/Runtime/HTTP/HTTPManager.h
--- a/Engine/Plugins/Media/ElectraPlayer/Source/ElectraPlayerRuntime/Private/Runtime/HTTP/HTTPManager.h
+++ b/Engine/Plugins/Media/ElectraPlayer/Source/ElectraPlayerRuntime/Private/Runtime/HTTP/HTTPManager.h
@@ -62,6 +62,20 @@ namespace Electra
 		{
 			FString	Header;
 			FString	Value;
+
+			void SetFromString(const FString& InString)
+			{
+				int32 ColonPos;
+				if (InString.FindChar(TCHAR(':'), ColonPos))
+				{
+					Header = InString.Left(ColonPos);
+					Value = InString.Mid(ColonPos + 2);
+				}
+			}
+			static void ParseFromString(FHTTPHeader& OutHeader, const FString& InString)
+			{
+				OutHeader.SetFromString(InString);
+			}
 		};
 
 		struct FConnectionInfo
@@ -385,6 +399,17 @@ namespace Electra
 				int64			EndIncluding = -1;
 				int64			DocumentSize = -1;
 			};
+
+			void AddFromHeaderList(const TArray<FString>& InHeaderList)
+			{
+				for(int32 i=0; i<InHeaderList.Num(); ++i)
+				{
+					HTTP::FHTTPHeader h;
+					h.SetFromString(InHeaderList[i]);
+					RequestHeaders.Emplace(MoveTemp(h));
+				}
+			}
+
 			FString								URL;							//!< URL
 			FString								Verb;							//!< GET (default if not set), HEAD, OPTIONS,....
 			FRange								Range;							//!< Optional request range
@@ -393,6 +418,7 @@ namespace Electra
 			TMediaOptionalValue<FString>		AcceptEncoding;					//!< Optional accepted encoding
 			FTimeValue							ConnectTimeout;					//!< Optional timeout for connecting to the server
 			FTimeValue							NoDataTimeout;					//!< Optional timeout when no data is being received
+			TArray<uint8>						PostData;						//!< Data for POST
 		};
 
 
@@ -403,10 +429,11 @@ namespace Electra
 			HTTP::FConnectionInfo				ConnectionInfo;
 			TWeakPtrTS<FReceiveBuffer>			ReceiveBuffer;
 			TWeakPtrTS<FProgressListener>		ProgressListener;
+			bool								bAutoRemoveWhenComplete = false;
 		};
 
-		virtual void AddRequest(TSharedPtrTS<FRequest> Request) = 0;
-		virtual void RemoveRequest(TSharedPtrTS<FRequest> Request) = 0;
+		virtual void AddRequest(TSharedPtrTS<FRequest> Request, bool bAutoRemoveWhenComplete) = 0;
+		virtual void RemoveRequest(TSharedPtrTS<FRequest> Request, bool bDoNotWaitForRemoval) = 0;
 
 	};
 
diff --git a/Engine/Plugins/Media/ElectraPlayer/Source/ElectraPlayerRuntime/Private/Runtime/Player/AdaptiveStreamingPlayer.cpp b/Engine/Plugins/Media/ElectraPlayer/Source/ElectraPlayerRuntime/Private/Runtime/Player/AdaptiveStreamingPlayer.cpp
--- a/Engine/Plugins/Media/ElectraPlayer/Source/ElectraPlayerRuntime/Private/Runtime/Player/AdaptiveStreamingPlayer.cpp
+++ b/Engine/Plugins/Media/ElectraPlayer/Source/ElectraPlayerRuntime/Private/Runtime/Player/AdaptiveStreamingPlayer.cpp
@@ -9,6 +9,7 @@
 #include "Player/PlayerLicenseKey.h"
 #include "Player/AdaptivePlayerOptionKeynames.h"
 #include "Utilities/Utilities.h"
+#include "Player/DRM/DRMManager.h"
 
 #include "HAL/LowLevelMemTracker.h"
 
@@ -183,6 +184,9 @@ bool FAdaptiveStreamingPlayer::Initialize(const FParamDict& Options)
 	// Get the HTTP manager. This is a shared instance for all players.
 	HttpManager = IElectraHttpManager::Create();
 
+	// Create the DRM manager.
+	DrmManager = FDRMManager::Create(this);
+
 	// Create an entity cache.
 	EntityCache = IPlayerEntityCache::Create(this, Options);
 
@@ -877,6 +881,13 @@ FParamDict& FAdaptiveStreamingPlayer::GetOptions()
 	return PlayerOptions;
 }
 
+TSharedPtrTS<FDRMManager> FAdaptiveStreamingPlayer::GetDRMManager()
+{
+	return DrmManager;
+}
+
+
+
 void FAdaptiveStreamingPlayer::GetStreamBufferStats(FAccessUnitBufferInfo& OutBufferStats, EStreamType ForStream)
 {
 	FMediaCriticalSection::ScopedLock lock(DiagnosticsCriticalSection);
@@ -1357,6 +1368,11 @@ void FAdaptiveStreamingPlayer::WorkerThreadFN()
 			{
 				EntityCache->HandleEntityExpiration();
 			}
+			// Handle completed DRM requests.
+			if (DrmManager.IsValid())
+			{
+				DrmManager->Tick();
+			}
 		}
 	}
 
@@ -3170,6 +3186,12 @@ void FAdaptiveStreamingPlayer::InternalClose()
 	//PlaybackState.Reset();
 	Manifest.Reset();
 
+	if (DrmManager.IsValid())
+	{
+		DrmManager->Close();
+		DrmManager.Reset();
+	}
+
 	HttpManager.Reset();
 	EntityCache.Reset();
 	RemoveMetricsReceiver(StreamSelector.Get());
diff --git a/Engine/Plugins/Media/ElectraPlayer/Source/ElectraPlayerRuntime/Private/Runtime/Player/AdaptiveStreamingPlayerInternal.h b/Engine/Plugins/Media/ElectraPlayer/Source/ElectraPlayerRuntime/Private/Runtime/Player/AdaptiveStreamingPlayerInternal.h
--- a/Engine/Plugins/Media/ElectraPlayer/Source/ElectraPlayerRuntime/Private/Runtime/Player/AdaptiveStreamingPlayerInternal.h
+++ b/Engine/Plugins/Media/ElectraPlayer/Source/ElectraPlayerRuntime/Private/Runtime/Player/AdaptiveStreamingPlayerInternal.h
@@ -18,6 +18,8 @@
 #include "Player/PlayerStreamFilter.h"
 #include "Player/PlayerEntityCache.h"
 
+#include "ElectraCDM.h"
+
 #include "InfoLog.h"
 
 #define INTERR_ALL_STREAMS_HAVE_FAILED			1
@@ -726,6 +728,7 @@ private:
 	virtual TSharedPtrTS<IPlayerEntityCache> GetEntityCache() override;
 	virtual IAdaptiveStreamingPlayerAEMSHandler* GetAEMSEventHandler() override;
 	virtual FParamDict& GetOptions() override;
+	virtual TSharedPtrTS<FDRMManager> GetDRMManager() override;
 
 	// Methods from IPlayerStreamFilter
 	virtual bool CanDecodeStream(const FStreamCodecInformation& InStreamCodecInfo) const override;
@@ -1347,6 +1350,7 @@ private:
 	TSharedPtrTS<IElectraHttpManager>									HttpManager;
 	TSharedPtrTS<IPlayerEntityCache>									EntityCache;
 
+	TSharedPtrTS<FDRMManager>											DrmManager;
 
 	TMediaQueueDynamic<TSharedPtrTS<FErrorDetail>>						ErrorQueue;
 
diff --git a/Engine/Plugins/Media/ElectraPlayer/Source/ElectraPlayerRuntime/Private/Runtime/Player/AdaptiveStreamingPlayerResourceRequest.cpp b/Engine/Plugins/Media/ElectraPlayer/Source/ElectraPlayerRuntime/Private/Runtime/Player/AdaptiveStreamingPlayerResourceRequest.cpp
--- a/Engine/Plugins/Media/ElectraPlayer/Source/ElectraPlayerRuntime/Private/Runtime/Player/AdaptiveStreamingPlayerResourceRequest.cpp
+++ b/Engine/Plugins/Media/ElectraPlayer/Source/ElectraPlayerRuntime/Private/Runtime/Player/AdaptiveStreamingPlayerResourceRequest.cpp
@@ -26,7 +26,7 @@ FHTTPResourceRequest::~FHTTPResourceRequest()
 	Cancel();
 	if (bWasAdded && Request.IsValid())
 	{
-		PlayerSessionServices->GetHTTPManager()->RemoveRequest(Request);
+		PlayerSessionServices->GetHTTPManager()->RemoveRequest(Request, true);
 	}
 	Request.Reset();
 }
@@ -45,7 +45,7 @@ void FHTTPResourceRequest::StartGet(IPlayerSessionServices* InPlayerSessionServi
 	else
 	{
 		bWasAdded = true;
-		PlayerSessionServices->GetHTTPManager()->AddRequest(Request);
+		PlayerSessionServices->GetHTTPManager()->AddRequest(Request, false);
 	}
 }
 
@@ -56,7 +56,7 @@ void FHTTPResourceRequest::Cancel()
 	ReceiveBuffer.Reset();
 	if (bWasAdded && Request.IsValid())
 	{
-		PlayerSessionServices->GetHTTPManager()->RemoveRequest(Request);
+		PlayerSessionServices->GetHTTPManager()->RemoveRequest(Request, true);
 		bWasAdded = false;
 	}
 }
@@ -71,7 +71,7 @@ void FHTTPResourceRequest::StaticDataReady()
 		{
 			// Do the actual HTTP request now.
 			bWasAdded = true;
-			PlayerSessionServices->GetHTTPManager()->AddRequest(Request);
+			PlayerSessionServices->GetHTTPManager()->AddRequest(Request, false);
 		}
 		else
 		{
diff --git a/Engine/Plugins/Media/ElectraPlayer/Source/ElectraPlayerRuntime/Private/Runtime/Player/AdaptiveStreamingPlayerResourceRequest.h b/Engine/Plugins/Media/ElectraPlayer/Source/ElectraPlayerRuntime/Private/Runtime/Player/AdaptiveStreamingPlayerResourceRequest.h
--- a/Engine/Plugins/Media/ElectraPlayer/Source/ElectraPlayerRuntime/Private/Runtime/Player/AdaptiveStreamingPlayerResourceRequest.h
+++ b/Engine/Plugins/Media/ElectraPlayer/Source/ElectraPlayerRuntime/Private/Runtime/Player/AdaptiveStreamingPlayerResourceRequest.h
@@ -68,12 +68,18 @@ public:
 
 	virtual FHTTPResourceRequest& Verb(const FString& InVerb)
 	{ Request->Parameters.Verb = InVerb; return *this; }
+
+	virtual FHTTPResourceRequest& PostData(TArray<uint8> InPostData)
+	{ Request->Parameters.PostData = MoveTemp(InPostData); return *this; }
 	
 	virtual FHTTPResourceRequest& Range(const FString& InRange)
 	{ Request->Parameters.Range.Set(InRange); return *this; }
 	
 	virtual FHTTPResourceRequest& Headers(const TArray<HTTP::FHTTPHeader>& InHeaders)
 	{ Request->Parameters.RequestHeaders = InHeaders; return *this; }
+
+	virtual FHTTPResourceRequest& Headers(const TArray<FString>& InHeaders)
+	{ Request->Parameters.AddFromHeaderList(InHeaders); return *this; }
 	
 	virtual FHTTPResourceRequest& AcceptEncoding(const FString& InAcceptEncoding)
 	{ Request->Parameters.AcceptEncoding.Set(InAcceptEncoding); return *this; }
diff --git a/Engine/Plugins/Media/ElectraPlayer/Source/ElectraPlayerRuntime/Private/Runtime/Player/DASH/MPDElementsDASH.h b/Engine/Plugins/Media/ElectraPlayer/Source/ElectraPlayerRuntime/Private/Runtime/Player/DASH/MPDElementsDASH.h
--- a/Engine/Plugins/Media/ElectraPlayer/Source/ElectraPlayerRuntime/Private/Runtime/Player/DASH/MPDElementsDASH.h
+++ b/Engine/Plugins/Media/ElectraPlayer/Source/ElectraPlayerRuntime/Private/Runtime/Player/DASH/MPDElementsDASH.h
@@ -359,7 +359,15 @@ public:
 	const FString& GetValue() const			{ return Value; }
 	const FString& GetID() const			{ return ID; }
 
-	const TArray<TSharedPtrTS<IDashMPDElement>>& GetWellKnownDescriptors() const	{ return WellKnownDescriptors; }
+	const TArray<TSharedPtrTS<IDashMPDElement>>& GetWellKnownDescriptors() const
+	{ return WellKnownDescriptors; }
+
+	// Non-DASH attributes and other XML elements as JSON. Used for <ContentProtection> elements
+	// where these elements need to be passed to a CDM system.
+	const FString& GetCustomElementAndAttributeJSON() const
+	{ return CustomElementAndAttributeJSON; }
+	void SetCustomElementAndAttributeJSON(const FString& InCustomElementAndAttributeJSON)
+	{ CustomElementAndAttributeJSON = InCustomElementAndAttributeJSON; }
 
 private:
 	virtual bool ProcessElement(FManifestParserDASH* Builder, const TCHAR* ElementName, const TCHAR* ElementData, int32 XmlFileLineNumber) override;
@@ -371,6 +379,8 @@ private:
 	FString SchemeIdUri;
 	FString Value;
 	FString ID;
+
+	FString CustomElementAndAttributeJSON;
 };
 
 
diff --git a/Engine/Plugins/Media/ElectraPlayer/Source/ElectraPlayerRuntime/Private/Runtime/Player/DASH/ManifestBuilderDASH.cpp b/Engine/Plugins/Media/ElectraPlayer/Source/ElectraPlayerRuntime/Private/Runtime/Player/DASH/ManifestBuilderDASH.cpp
--- a/Engine/Plugins/Media/ElectraPlayer/Source/ElectraPlayerRuntime/Private/Runtime/Player/DASH/ManifestBuilderDASH.cpp
+++ b/Engine/Plugins/Media/ElectraPlayer/Source/ElectraPlayerRuntime/Private/Runtime/Player/DASH/ManifestBuilderDASH.cpp
@@ -24,6 +24,7 @@
 #include "Utilities/Utilities.h"
 #include "Utilities/TimeUtilities.h"
 
+#include "Player/DRM/DRMManager.h"
 
 #define ERRCODE_DASH_MPD_BUILDER_INTERNAL							1
 #define ERRCODE_DASH_MPD_BUILDER_UNSUPPORTED_PROFILE				100
@@ -87,8 +88,28 @@ namespace
 		TEXT("urn:mpeg:dash:urlparam:2016"),
 	};
 
+    const TCHAR* const Scheme_urn_mpeg_dash_mp4protection_2011 = TEXT("urn:mpeg:dash:mp4protection:2011");
 }
 
+
+namespace DASHAttributeHelpers
+{
+	const IDashMPDElement::FXmlAttribute* GetAttribute(const TSharedPtrTS<FDashMPD_DescriptorType>& InDescriptor, const TCHAR* InAttribute, const TCHAR* InOptionalNamespace)
+	{
+		const TArray<IDashMPDElement::FXmlAttribute>& Attributes = InDescriptor->GetOtherAttributes();
+		FString NameWithNS = InOptionalNamespace ? FString::Printf(TEXT("%s:%s"), InOptionalNamespace, InAttribute) : FString(InAttribute);
+		for(int32 i=0; i<Attributes.Num(); ++i)
+		{
+			if (Attributes[i].GetName().Equals(NameWithNS) || Attributes[i].GetName().Equals(InAttribute))
+			{
+				return &Attributes[i];
+			}
+		}
+		return nullptr;
+	}
+};
+
+
 namespace DASHUrlHelpers
 {
 	ELECTRA_IMPL_DEFAULT_ERROR_METHODS(DASHMPDBuilder);
@@ -1340,10 +1361,10 @@ void FManifestDASHInternal::PreparePeriodAdaptationSets(TSharedPtrTS<FPeriod> Pe
 					continue;
 				}
 
-				// Encryption?
-				if (MPDRepresentation->GetContentProtections().Num() || MPDAdaptationSet->GetContentProtections().Num())
+				// Encryption on representation level is not supported. As per DASH-IF-IOP v4.3 encryption must be specified on adaptation set level.
+				if (MPDRepresentation->GetContentProtections().Num())
 				{
-					LogMessage(PlayerSessionServices, IInfoLog::ELevel::Warning, FString::Printf(TEXT("ContentProtection is not supported, ignoring this Representation.")));
+					LogMessage(PlayerSessionServices, IInfoLog::ELevel::Warning, FString::Printf(TEXT("ContentProtection in Representation is not supported. Must be on enclosing AdaptationSet! Ignoring this Representation.")));
 					continue;
 				}
 
@@ -1523,9 +1544,20 @@ void FManifestDASHInternal::PreparePeriodAdaptationSets(TSharedPtrTS<FPeriod> Pe
 				RepresentationQualityIndexMap.Add(MPDRepresentation->GetBandwidth(), Representation);
 			}
 
+
 			// If the adaptation set contains usable Representations we mark the AdaptationSet as usable as well.
 			if (AdaptationSet->Representations.Num())
 			{
+				// Encryption?
+				if (MPDAdaptationSet->GetContentProtections().Num())
+				{
+					if (!CanUseEncryptedAdaptation(AdaptationSet))
+					{
+						LogMessage(PlayerSessionServices, IInfoLog::ELevel::Warning, FString::Printf(TEXT("ContentProtection is not supported, ignoring this AdaptationSet.")));
+						continue;
+					}
+				}
+
 				RepresentationQualityIndexMap.KeySort([](int32 A, int32 B){return A<B;});
 				int32 CurrentQualityIndex = -1;
 				int32 CurrentQualityBitrate = -1;
@@ -1548,6 +1580,74 @@ void FManifestDASHInternal::PreparePeriodAdaptationSets(TSharedPtrTS<FPeriod> Pe
 }
 
 
+bool FManifestDASHInternal::CanUseEncryptedAdaptation(const TSharedPtrTS<FAdaptationSet>& InAdaptationSet)
+{
+	TSharedPtrTS<FDRMManager> DRMManager = PlayerSessionServices->GetDRMManager();
+	if (DRMManager.IsValid())
+	{
+		const TSharedPtrTS<FDashMPD_AdaptationSetType> MPDAdaptationSet = InAdaptationSet->AdaptationSet.Pin();
+		const TArray<TSharedPtrTS<FDashMPD_DescriptorType>>& ContentProtections = MPDAdaptationSet->GetContentProtections();
+		FString Mime = InAdaptationSet->GetCodec().GetMimeTypeWithCodecAndFeatures();
+		// See if there is a DASH scheme saying that common encryption is in use.
+		for(int32 i=0; i<ContentProtections.Num(); ++i)
+		{
+			if (ContentProtections[i]->GetSchemeIdUri().Equals(Scheme_urn_mpeg_dash_mp4protection_2011, ESearchCase::IgnoreCase))
+			{
+				InAdaptationSet->CommonEncryptionScheme = ContentProtections[i]->GetValue();
+				const IDashMPDElement::FXmlAttribute* default_KID = DASHAttributeHelpers::GetAttribute(ContentProtections[i], TEXT("default_KID"), TEXT("cenc"));
+				InAdaptationSet->DefaultKID = default_KID ? default_KID->GetValue() : FString();
+				break;
+			}
+		}
+		for(int32 i=0; i<ContentProtections.Num(); ++i)
+		{
+			TSharedPtr<ElectraCDM::IMediaCDMCapabilities, ESPMode::ThreadSafe> DRMCapabilities;
+			// Skip over the common encryption scheme which we handled in the first pass already.
+			if (ContentProtections[i]->GetSchemeIdUri().Equals(Scheme_urn_mpeg_dash_mp4protection_2011, ESearchCase::IgnoreCase))
+			{
+				continue;
+			}
+
+			// Get the scheme specific attributes and child elements as a JSON. If this does not exist yet
+			// create it and store with the descriptor for later use.
+			FString SchemeSpecificData = ContentProtections[i]->GetCustomElementAndAttributeJSON();
+			if (SchemeSpecificData.IsEmpty())
+			{
+				FString AttrPrefix, TextProp;
+				bool bNoNamespaces = false;
+				DRMManager->GetCDMCustomJSONPrefixes(ContentProtections[i]->GetSchemeIdUri(), ContentProtections[i]->GetValue(), AttrPrefix, TextProp, bNoNamespaces);
+				IManifestParserDASH::BuildJSONFromCustomElement(SchemeSpecificData, ContentProtections[i], false, bNoNamespaces, false, true, *AttrPrefix, *TextProp);
+				ContentProtections[i]->SetCustomElementAndAttributeJSON(SchemeSpecificData);
+			}
+
+			DRMCapabilities = DRMManager->GetCDMCapabilitiesForScheme(ContentProtections[i]->GetSchemeIdUri(), ContentProtections[i]->GetValue(), SchemeSpecificData);
+			if (DRMCapabilities.IsValid())
+			{
+				ElectraCDM::IMediaCDMCapabilities::ESupportResult Result;
+				Result = DRMCapabilities->SupportsType(Mime);
+				if (Result != ElectraCDM::IMediaCDMCapabilities::ESupportResult::Supported)
+				{
+					continue;
+				}
+				Result = DRMCapabilities->RequiresSecureDecoder(Mime);
+				if (Result == ElectraCDM::IMediaCDMCapabilities::ESupportResult::SecureDecoderRequired)
+				{
+					//LogMessage(PlayerSessionServices, IInfoLog::ELevel::Warning, FString::Printf(TEXT("Use of secure decoders is not supported.")));
+					continue;
+				}
+
+				FAdaptationSet::FContentProtection Prot;
+				Prot.Descriptor = ContentProtections[i];
+				Prot.DefaultKID = InAdaptationSet->GetDefaultKID();
+				Prot.CommonScheme = InAdaptationSet->GetCommonEncryptionScheme();
+				InAdaptationSet->PossibleContentProtections.Emplace(MoveTemp(Prot));
+			}
+		}
+	}
+	return InAdaptationSet->PossibleContentProtections.Num() > 0;
+}
+
+
 
 void FManifestDASHInternal::SendEventsFromAllPeriodEventStreams(TSharedPtrTS<FPeriod> InPeriod)
 {
diff --git a/Engine/Plugins/Media/ElectraPlayer/Source/ElectraPlayerRuntime/Private/Runtime/Player/DASH/ManifestDASH.cpp b/Engine/Plugins/Media/ElectraPlayer/Source/ElectraPlayerRuntime/Private/Runtime/Player/DASH/ManifestDASH.cpp
--- a/Engine/Plugins/Media/ElectraPlayer/Source/ElectraPlayerRuntime/Private/Runtime/Player/DASH/ManifestDASH.cpp
+++ b/Engine/Plugins/Media/ElectraPlayer/Source/ElectraPlayerRuntime/Private/Runtime/Player/DASH/ManifestDASH.cpp
@@ -16,6 +16,7 @@
 #include "Player/DASH/OptionKeynamesDASH.h"
 #include "Player/PlayerEntityCache.h"
 #include "Player/AdaptivePlayerOptionKeynames.h"
+#include "Player/DRM/DRMManager.h"
 
 
 namespace Electra
@@ -23,6 +24,7 @@ namespace Electra
 
 #define ERRCODE_DASH_MPD_INTERNAL						1
 #define ERRCODE_DASH_MPD_BAD_REPRESENTATION				1000
+#define ERRCODE_DASH_DRM_ERROR							2000
 
 
 namespace DashUtils
@@ -95,6 +97,10 @@ public:
 
 	virtual ~FDASHPlayPeriod()
 	{
+		if (DrmClient.IsValid())
+		{
+			DrmClient->UnregisterEventListener(PlayerSessionServices->GetDRMManager());
+		}
 	}
 
 	//----------------------------------------------
@@ -126,6 +132,9 @@ private:
 	};
 	IManifest::FResult GetNextOrRetrySegment(TSharedPtrTS<IStreamSegment>& OutSegment, TSharedPtrTS<const IStreamSegment> InCurrentSegment, ENextSegType InNextType);
 
+	bool PrepareDRM(const TArray<FManifestDASHInternal::FAdaptationSet::FContentProtection>& InContentProtections);
+
+
 	IPlayerSessionServices* PlayerSessionServices = nullptr;
 	EReadyState ReadyState = EReadyState::NotReady;
 	FStreamPreferences StreamPreferences;
@@ -139,6 +148,8 @@ private:
 
 	FString ActiveVideoRepresentationID;
 	FString ActiveAudioRepresentationID;
+
+	TSharedPtr<ElectraCDM::IMediaCDMClient, ESPMode::ThreadSafe> DrmClient;
 };
 
 /***************************************************************************************************************************************************/
@@ -458,6 +469,8 @@ void FDASHPlayPeriod::PrepareForPlay()
 	TSharedPtrTS<FManifestDASHInternal::FPeriod> Period = Manifest.IsValid() ? Manifest->GetPeriodByUniqueID(PeriodID) : nullptr;
 	if (Period.IsValid())
 	{
+		TArray<FManifestDASHInternal::FAdaptationSet::FContentProtection> ContentProtections;
+
 		// Prepare the adaptation sets and periods.
 		Manifest->PreparePeriodAdaptationSets(Period, false);
 	
@@ -471,6 +484,10 @@ void FDASHPlayPeriod::PrepareForPlay()
 
 			TSharedPtrTS<IPlaybackAssetRepresentation> VideoRepr = GetRepresentationFromAdaptationByMaxBandwidth(VideoAS, 2*1000*1000);
 			ActiveVideoRepresentationID = VideoRepr->GetUniqueIdentifier();
+			
+			// Add encryption schemes, if any.
+			const FManifestDASHInternal::FAdaptationSet* Adapt = static_cast<const FManifestDASHInternal::FAdaptationSet*>(VideoAS.Get());
+			ContentProtections.Append(Adapt->GetPossibleContentProtections());
 		}
 
 		int32 NumAudioAdaptationSets = Period->GetNumberOfAdaptationSets(EStreamType::Audio);
@@ -482,16 +499,64 @@ void FDASHPlayPeriod::PrepareForPlay()
 
 			TSharedPtrTS<IPlaybackAssetRepresentation> AudioRepr = GetRepresentationFromAdaptationByMaxBandwidth(AudioAS, 128 * 1000);
 			ActiveAudioRepresentationID = AudioRepr->GetUniqueIdentifier();
+
+			// Add encryption schemes, if any.
+			const FManifestDASHInternal::FAdaptationSet* Adapt = static_cast<const FManifestDASHInternal::FAdaptationSet*>(AudioAS.Get());
+			ContentProtections.Append(Adapt->GetPossibleContentProtections());
 		}
 
 		// Emit all <EventStream> events of the period to the AEMS event handler.
 		Manifest->SendEventsFromAllPeriodEventStreams(Period);
 
-	//	ReadyState = EReadyState::Preparing;
-		ReadyState = EReadyState::IsReady;
+		// Prepare the DRM system for decryption.
+		if (PrepareDRM(ContentProtections))
+		{
+			//ReadyState = EReadyState::Preparing;
+			ReadyState = EReadyState::IsReady;
+		}
+		else
+		{
+			// Set state to preparing to prevent the player from progressing while
+			// the posted error works its magic.
+			ReadyState = EReadyState::Preparing;
+		}
+	}
+}
+
+bool FDASHPlayPeriod::PrepareDRM(const TArray<FManifestDASHInternal::FAdaptationSet::FContentProtection>& InContentProtections)
+{
+	if (InContentProtections.Num())
+	{
+		static const TCHAR* const Delimiters[] { TEXT(" "), TEXT("\t"), TEXT("\n"), TEXT("\r") };
+		// Set up DRM CRM candidates and settle on one.
+		TArray<ElectraCDM::IMediaCDM::FCDMCandidate> Candidates;
+		ElectraCDM::IMediaCDM::FCDMCandidate cand;
+		for(int32 i=0; i<InContentProtections.Num(); ++i)
+		{
+			cand.SchemeId = InContentProtections[i].Descriptor->GetSchemeIdUri();
+			cand.Value = InContentProtections[i].Descriptor->GetValue();
+			cand.CommonScheme = InContentProtections[i].CommonScheme;
+			cand.AdditionalElements = InContentProtections[i].Descriptor->GetCustomElementAndAttributeJSON();
+			InContentProtections[i].DefaultKID.ParseIntoArray(cand.DefaultKIDs, Delimiters, UE_ARRAY_COUNT(Delimiters), true);
+			Candidates.Emplace(MoveTemp(cand));
+		}
+		ElectraCDM::ECDMError Result = PlayerSessionServices->GetDRMManager()->CreateDRMClient(DrmClient, Candidates);
+		if (Result == ElectraCDM::ECDMError::Success && DrmClient.IsValid())
+		{
+			DrmClient->RegisterEventListener(PlayerSessionServices->GetDRMManager());
+			DrmClient->PrepareLicenses();
+			return true;
+		}
+		else
+		{
+			PostError(PlayerSessionServices, FString::Printf(TEXT("Failed to create DRM client with error %d"), (int32)Result), ERRCODE_DASH_DRM_ERROR);
+			return false;
+		}
 	}
+	return true;
 }
 
+
 TSharedPtrTS<ITimelineMediaAsset> FDASHPlayPeriod::GetMediaAsset() const
 {
 	TSharedPtrTS<FManifestDASHInternal> Manifest = GetCurrentManifest();
@@ -707,6 +772,10 @@ IManifest::FResult FDASHPlayPeriod::GetStartingSegment(TSharedPtrTS<IStreamSegme
 					StartSegmentRequest->SAET = SegmentRequest->SAET;
 				}
 
+				// Encryption stuff
+				SegmentRequest->DrmClient = DrmClient;
+				SegmentRequest->DrmMimeType = Repr->GetCodecInformation().GetMimeTypeWithCodecAndFeatures();
+
 				StartSegmentRequest->DependentStreams.Emplace(MoveTemp(SegmentRequest));
 				bAllStreamsAtEOS = false;
 			}
@@ -895,6 +964,10 @@ IManifest::FResult FDASHPlayPeriod::GetNextOrRetrySegment(TSharedPtrTS<IStreamSe
 			SegmentRequest->bWarnedAboutTimescale = Current->bWarnedAboutTimescale;
 		}
 
+		// Encryption stuff
+		SegmentRequest->DrmClient = DrmClient;
+		SegmentRequest->DrmMimeType = Repr->GetCodecInformation().GetMimeTypeWithCodecAndFeatures();
+
 		SegmentRequest->Segment = MoveTemp(SegmentInfo);
 		OutSegment = MoveTemp(SegmentRequest);
 		return IManifest::FResult(IManifest::FResult::EType::Found);
diff --git a/Engine/Plugins/Media/ElectraPlayer/Source/ElectraPlayerRuntime/Private/Runtime/Player/DASH/ManifestParserDASH.cpp b/Engine/Plugins/Media/ElectraPlayer/Source/ElectraPlayerRuntime/Private/Runtime/Player/DASH/ManifestParserDASH.cpp
--- a/Engine/Plugins/Media/ElectraPlayer/Source/ElectraPlayerRuntime/Private/Runtime/Player/DASH/ManifestParserDASH.cpp
+++ b/Engine/Plugins/Media/ElectraPlayer/Source/ElectraPlayerRuntime/Private/Runtime/Player/DASH/ManifestParserDASH.cpp
@@ -3255,5 +3255,176 @@ static bool ParseUIntArray(TArray<uint32>& OutArray, const TCHAR* Delimiter, FMa
 
 } // namespace anonymous
 
+
+/*********************************************************************************************************************/
+/*********************************************************************************************************************/
+/*********************************************************************************************************************/
+
+namespace IManifestParserDASH
+{
+namespace
+{
+	FString RemoveNamespacePrefix(const FString& InString, bool bRemoveNS)
+	{
+		int32 ColonPos;
+		if (bRemoveNS && InString.FindLastChar(TCHAR(':'), ColonPos))
+		{
+			return InString.RightChop(ColonPos + 1);
+		}
+		return InString;
+	}
+
+	FString MakeJSONElement(const FString& Key, const FString& Value, bool bRemoveNamespaces, const TCHAR* Prefix)
+	{
+		return FString::Printf(TEXT("\"%s%s\":\"%s\""), Prefix, *RemoveNamespacePrefix(Key, bRemoveNamespaces), *Value.ReplaceQuotesWithEscapedQuotes());
+	}
+
+	FString MakeJSONAttribute(const IDashMPDElement::FXmlAttribute& Attribute, bool bRemoveNamespaces, const TCHAR* Prefix)
+	{
+		return MakeJSONElement(Attribute.GetName(), Attribute.GetValue(), bRemoveNamespaces, Prefix);
+	}
+
+	void BuildJSONFromElementRec(FString& OutJSON, TSharedPtrTS<IDashMPDElement> InElement, bool bRemoveNamespaces, bool bForce1ElementArrays, bool bTerseObjects, const TCHAR* AttributePrefix, const TCHAR* TextPropertyName)
+	{
+		const TArray<TSharedPtrTS<FDashMPD_OtherType>>& Children = InElement->GetOtherChildren();
+		const TArray<IDashMPDElement::FXmlAttribute>& Attributes = InElement->GetOtherAttributes();
+		bool bCanBeTerse = Children.Num() == 0 && Attributes.Num() == 0;
+
+		bool bNeedElementBraces = !bCanBeTerse || !bTerseObjects;
+
+		if (bNeedElementBraces)
+		{
+			OutJSON.AppendChar(TCHAR('{'));
+		}
+
+		// Add attributes
+		for(int32 i=0; i<Attributes.Num(); ++i)
+		{
+			OutJSON.Append(MakeJSONAttribute(Attributes[i], bRemoveNamespaces, AttributePrefix));
+			if (i+1 <Attributes.Num())
+			{
+				OutJSON.AppendChar(TCHAR(','));
+			}
+		}
+
+		// Add element data if it exists.
+		bool bNeedComma = Attributes.Num() > 0;
+		if (InElement->GetData().Len())
+		{
+			if (bNeedComma)
+			{
+				OutJSON.AppendChar(TCHAR(','));
+			}
+			if (bNeedElementBraces)
+			{
+				OutJSON.Append(MakeJSONElement(TextPropertyName, InElement->GetData(), false, TEXT("")));
+			}
+			else
+			{
+				OutJSON.Append(FString::Printf(TEXT("\"%s\""), *InElement->GetData().ReplaceQuotesWithEscapedQuotes()));
+			}
+			bNeedComma = true;
+		}
+		else if (!bNeedElementBraces)
+		{
+			// Need to add an empty string
+			OutJSON.Append(TEXT("\"\""));
+		}
+
+
+		// Collect the children into a "map" in order to create a JSON array for those that have the same name.
+		// As we need to maintain string case we can't use a TMap to do this.
+		struct FCaseSensitiveChildElements
+		{
+			FString Key;
+			TArray<TSharedPtrTS<FDashMPD_OtherType>> List;
+		};
+		TArray<FCaseSensitiveChildElements> ChildrenMap;
+		for(int32 i=0; i<Children.Num(); ++i)
+		{
+			FString CaseKey = RemoveNamespacePrefix(Children[i]->GetName(), bRemoveNamespaces);
+			bool bFound = false;
+			for(int32 j=0; j<ChildrenMap.Num(); ++j)
+			{
+				if (ChildrenMap[j].Key.Equals(CaseKey, ESearchCase::CaseSensitive))
+				{
+					bFound = true;
+					ChildrenMap[j].List.Add(Children[i]);
+					break;
+				}
+			}
+			if (!bFound)
+			{
+				FCaseSensitiveChildElements NewCaseSensitiveList;
+				NewCaseSensitiveList.Key = CaseKey;
+				NewCaseSensitiveList.List.Add(Children[i]);
+				ChildrenMap.Emplace(MoveTemp(NewCaseSensitiveList));
+			}
+		}
+
+		for(int32 nC=0; nC<ChildrenMap.Num(); ++nC)
+		{
+			if (bNeedComma)
+			{
+				OutJSON.AppendChar(TCHAR(','));
+			}
+			bNeedComma = true;
+			const TArray<TSharedPtrTS<FDashMPD_OtherType>>& ArrayElements = ChildrenMap[nC].List;
+			int32 NumArrayElements = ArrayElements.Num();
+			if (NumArrayElements == 1 && !bForce1ElementArrays)
+			{
+				OutJSON.Append(FString::Printf(TEXT("\"%s\":"), *ChildrenMap[nC].Key));
+			}
+			else
+			{
+				OutJSON.Append(FString::Printf(TEXT("\"%s\":["), *ChildrenMap[nC].Key));
+			}
+			for(int32 i=0; i<NumArrayElements; ++i)
+			{
+				BuildJSONFromElementRec(OutJSON, ArrayElements[i], bRemoveNamespaces, bForce1ElementArrays, bTerseObjects, AttributePrefix, TextPropertyName);
+				if (i+1 < NumArrayElements)
+				{
+					OutJSON.AppendChar(TCHAR(','));
+				}
+			}
+			if (!(NumArrayElements == 1 && !bForce1ElementArrays))
+			{
+				OutJSON.AppendChar(TCHAR(']'));
+			}
+		}
+
+		if (bNeedElementBraces)
+		{
+			OutJSON.AppendChar(TCHAR('}'));
+		}
+	}
+}
+
+void BuildJSONFromCustomElement(FString& OutJSON, TSharedPtrTS<IDashMPDElement> InElement, bool bIncludeStartElement, bool bRemoveNamespaces, bool bForce1ElementArrays, bool bTerseObjects, const TCHAR* AttributePrefix, const TCHAR* TextPropertyName)
+{
+	check(AttributePrefix);
+	check(TextPropertyName);
+	if (InElement.IsValid())
+	{
+		if (bIncludeStartElement)
+		{
+			OutJSON.Append(FString::Printf(TEXT("{\"%s\":"), *RemoveNamespacePrefix(InElement->GetName(), bRemoveNamespaces)));
+		}
+
+		BuildJSONFromElementRec(OutJSON, InElement, bRemoveNamespaces, bForce1ElementArrays, bTerseObjects, AttributePrefix, TextPropertyName);
+
+		if (bIncludeStartElement)
+		{
+			OutJSON.AppendChar(TCHAR('}'));
+		}
+	}
+	else
+	{
+		OutJSON.Append(TEXT("{}"));
+	}
+}
+
+}
+
 } // namespace Electra
 
diff --git a/Engine/Plugins/Media/ElectraPlayer/Source/ElectraPlayerRuntime/Private/Runtime/Player/DASH/ManifestParserDASH.h b/Engine/Plugins/Media/ElectraPlayer/Source/ElectraPlayerRuntime/Private/Runtime/Player/DASH/ManifestParserDASH.h
--- a/Engine/Plugins/Media/ElectraPlayer/Source/ElectraPlayerRuntime/Private/Runtime/Player/DASH/ManifestParserDASH.h
+++ b/Engine/Plugins/Media/ElectraPlayer/Source/ElectraPlayerRuntime/Private/Runtime/Player/DASH/ManifestParserDASH.h
@@ -14,6 +14,9 @@ namespace Electra
 namespace IManifestParserDASH
 {
 FErrorDetail BuildFromMPD(FDashMPD_RootEntities& OutRootEntities, TArray<TWeakPtrTS<IDashMPDElement>>& OutXLinkElements, TCHAR* InOutMPDXML, const TCHAR* InExpectedRootElement, IPlayerSessionServices* InPlayerSessionServices);
+
+void BuildJSONFromCustomElement(FString& OutJSON, TSharedPtrTS<IDashMPDElement> InElement, bool bIncludeStartElement, bool bRemoveNamespaces, bool bForce1ElementArrays, bool bTerseObjects, const TCHAR* AttributePrefix, const TCHAR* TextPropertyName);
+
 };
 
 } // namespace Electra
diff --git a/Engine/Plugins/Media/ElectraPlayer/Source/ElectraPlayerRuntime/Private/Runtime/Player/DASH/PlaylistReaderDASH_Internal.h b/Engine/Plugins/Media/ElectraPlayer/Source/ElectraPlayerRuntime/Private/Runtime/Player/DASH/PlaylistReaderDASH_Internal.h
--- a/Engine/Plugins/Media/ElectraPlayer/Source/ElectraPlayerRuntime/Private/Runtime/Player/DASH/PlaylistReaderDASH_Internal.h
+++ b/Engine/Plugins/Media/ElectraPlayer/Source/ElectraPlayerRuntime/Private/Runtime/Player/DASH/PlaylistReaderDASH_Internal.h
@@ -228,6 +228,7 @@ public:
 		void GetSegmentInformation(TArray<IManifest::IPlayPeriod::FSegmentInformation>& OutSegmentInformation, FTimeValue& OutAverageSegmentDuration, TSharedPtrTS<const IStreamSegment> CurrentSegment, const FTimeValue& LookAheadTime, const TSharedPtrTS<IPlaybackAssetAdaptationSet>& AdaptationSet);
 
 
+
 		//----------------------------------------------
 		// Methods from IPlaybackAssetRepresentation
 		//
@@ -314,6 +315,30 @@ public:
 			return nullptr;
 		}
 
+		struct FContentProtection
+		{
+			TSharedPtrTS<FDashMPD_DescriptorType> Descriptor;
+			FString DefaultKID;
+			FString CommonScheme;
+		};
+		FString GetMimeType()
+		{
+			TSharedPtrTS<FDashMPD_AdaptationSetType> MPDAdaptationSet = AdaptationSet.Pin();
+			return MPDAdaptationSet.IsValid() ? MPDAdaptationSet->GetMimeType() : FString();
+		}
+		FString GetMimeTypeWithCodecs()
+		{
+			TSharedPtrTS<FDashMPD_AdaptationSetType> MPDAdaptationSet = AdaptationSet.Pin();
+			if (MPDAdaptationSet.IsValid())
+			{
+				return FString::Printf(TEXT("%s; codecs=\"%s\""), *MPDAdaptationSet->GetMimeType(), *GetListOfCodecs());
+			}
+			return FString();
+		}
+		const TArray<FContentProtection>& GetPossibleContentProtections() const { return PossibleContentProtections; }
+		const FString& GetCommonEncryptionScheme() const { return CommonEncryptionScheme; }
+		const FString& GetDefaultKID() const { return DefaultKID; }
+
 		//----------------------------------------------
 		// Methods from IPlaybackAssetAdaptationSet
 		//
@@ -358,6 +383,10 @@ public:
 		int32 IndexOfSelf = 0;
 		bool bIsUsable = false;
 		bool bIsEnabled = true;
+		// Encryption related
+		TArray<FContentProtection> PossibleContentProtections;
+		FString CommonEncryptionScheme;
+		FString DefaultKID;
 	};
 
 	class FPeriod : public ITimelineMediaAsset
@@ -636,6 +665,8 @@ private:
 	FTimeValue CalculateDistanceToLiveEdge() const;
 	FTimeRange GetPlayTimesFromURI() const;
 
+	bool CanUseEncryptedAdaptation(const TSharedPtrTS<FAdaptationSet>& InAdaptationSet);
+
 	IPlayerSessionServices* PlayerSessionServices = nullptr;
 
 	TArray<TWeakPtrTS<IDashMPDElement>> RemoteElementsToResolve;
diff --git a/Engine/Plugins/Media/ElectraPlayer/Source/ElectraPlayerRuntime/Private/Runtime/Player/DASH/StreamReaderFMP4DASH.cpp b/Engine/Plugins/Media/ElectraPlayer/Source/ElectraPlayerRuntime/Private/Runtime/Player/DASH/StreamReaderFMP4DASH.cpp
--- a/Engine/Plugins/Media/ElectraPlayer/Source/ElectraPlayerRuntime/Private/Runtime/Player/DASH/StreamReaderFMP4DASH.cpp
+++ b/Engine/Plugins/Media/ElectraPlayer/Source/ElectraPlayerRuntime/Private/Runtime/Player/DASH/StreamReaderFMP4DASH.cpp
@@ -14,6 +14,7 @@
 #include "Player/DASH/OptionKeynamesDASH.h"
 #include "Player/DASH/PlayerEventDASH.h"
 #include "Player/DASH/PlayerEventDASH_Internal.h"
+#include "Player/DRM/DRMManager.h"
 #include "Utilities/Utilities.h"
 
 #define INTERNAL_ERROR_INIT_SEGMENT_DOWNLOAD_ERROR					1
@@ -716,7 +717,7 @@ void FStreamReaderFMP4DASH::FStreamHandler::HandleRequest()
 
 			ProgressReportCount = 0;
 			DownloadCompleteSignal.Reset();
-			PlayerSessionService->GetHTTPManager()->AddRequest(HTTP);
+			PlayerSessionService->GetHTTPManager()->AddRequest(HTTP, false);
 
 			MP4Parser = IParserISO14496_12::CreateParser();
 			NumMOOFBoxesFound = 0;
@@ -725,6 +726,20 @@ void FStreamReaderFMP4DASH::FStreamHandler::HandleRequest()
 			uint32 TrackTimescale = 0;
 			bool bDone = false;
 			bool bIsFirstAU = true;
+
+			// Encrypted?
+			TSharedPtr<ElectraCDM::IMediaCDMDecrypter, ESPMode::ThreadSafe> Decrypter;
+			if (Request->DrmClient.IsValid())
+			{
+				if (Request->DrmClient->CreateDecrypter(Decrypter, Request->DrmMimeType) != ElectraCDM::ECDMError::Success)
+				{
+					bDone = true;
+					bHasErrored = true;
+					ds.FailureReason = FString::Printf(TEXT("Failed to create decrypter for segment. %s"), *Request->DrmClient->GetLastErrorMessage());
+					LogMessage(IInfoLog::ELevel::Error, ds.FailureReason);
+				}
+			}
+
 			while(!bDone && !HasErrored() && !HasReadBeenAborted())
 			{
 				UEMediaError parseError = MP4Parser->ParseHeader(this, this, PlayerSessionService, MP4InitSegment.Get());
@@ -766,6 +781,14 @@ void FStreamReaderFMP4DASH::FStreamHandler::HandleRequest()
 									}
 								}
 
+								// If the track uses encryption we update the DRM system with the PSSH boxes that are currently in use.
+								if (Decrypter.IsValid())
+								{
+									TArray<TArray<uint8>> PSSHBoxes;
+									Track->GetPSSHBoxes(PSSHBoxes, true, true);
+									Decrypter->UpdateInitDataFromMultiplePSSH(PSSHBoxes);
+								}
+
 								for(Error = TrackIterator->StartAtFirst(false); Error == UEMEDIA_ERROR_OK; Error = TrackIterator->Next())
 								{
 									// Get the DTS and PTS. Those are 0-based in a fragment and offset by the base media decode time of the fragment.
@@ -823,22 +846,8 @@ void FStreamReaderFMP4DASH::FStreamHandler::HandleRequest()
 									AccessUnit->PTS.SetFromND(AUPTS, TrackTimescale);
 									AccessUnit->PTS += TimeOffset;
 
-									#if 0
-									LogMessage(IInfoLog::ELevel::Verbose, FString::Printf(TEXT("[%s] %lld/%d/%u: DTS=%lld, PTS=%lld, dur=%lld, scale=%u, sync:%d, size=%lld, off=%lld"), GetStreamTypeName(Request->GetType()), (long long int)Request->Segment.Number, NumMOOFBoxesFound, TrackIterator->GetSampleNumber(),
-										#if 1
-										(long long int)AUDTS,
-										(long long int)AUPTS,
-										(long long int)AUDuration,
-										#else
-										(long long int)TrackIterator->GetDTS(),
-										(long long int)TrackIterator->GetPTS(),
-										(long long int)TrackIterator->GetDuration(),
-										#endif
-										TrackIterator->GetTimescale(),
-										TrackIterator->IsSyncSample()?1:0,
-										(long long int)TrackIterator->GetSampleSize(),
-										(long long int)TrackIterator->GetSampleFileOffset()));
-									#endif
+									ElectraCDM::FMediaCDMSampleInfo SampleEncryptionInfo;
+									bool bIsSampleEncrypted = TrackIterator->GetEncryptionInfo(SampleEncryptionInfo);
 
 									// There should not be any gaps!
 									int64 NumBytesToSkip = TrackIterator->GetSampleFileOffset() - GetCurrentOffset();
@@ -874,6 +883,31 @@ void FStreamReaderFMP4DASH::FStreamHandler::HandleRequest()
 										NextExpectedDTS = AccessUnit->DTS + Duration;
 										LastKnownAUDuration = Duration;
 										LastSuccessfulFilePos = GetCurrentOffset();
+
+										// If we need to decrypt we have to wait for the decrypter to become ready.
+										if (bIsSampleEncrypted && Decrypter.IsValid())
+										{
+											while(!bTerminate && !HasReadBeenAborted() && 
+												  (Decrypter->GetState() == ElectraCDM::ECDMState::WaitingForKey || Decrypter->GetState() == ElectraCDM::ECDMState::Idle))
+											{
+												FMediaRunnable::SleepMilliseconds(100);
+											}
+											ElectraCDM::ECDMError DecryptResult = ElectraCDM::ECDMError::Failure;
+											if (Decrypter->GetState() == ElectraCDM::ECDMState::Ready)
+											{
+												DecryptResult = Decrypter->DecryptInPlace((uint8*) AccessUnit->AUData, (int32) AccessUnit->AUSize, SampleEncryptionInfo);
+											}
+											if (DecryptResult != ElectraCDM::ECDMError::Success)
+											{
+												FAccessUnit::Release(AccessUnit);
+												AccessUnit = nullptr;
+												ds.FailureReason = FString::Printf(TEXT("Failed to decrypt segment \"%s\" with error %d (%s)"), *RequestURL, (int32)DecryptResult, *Decrypter->GetLastErrorMessage());
+												LogMessage(IInfoLog::ELevel::Error, ds.FailureReason);
+												bHasErrored = true;
+												bDone = true;
+												break;
+											}
+										}
 									}
 									else
 									{
@@ -968,7 +1002,7 @@ void FStreamReaderFMP4DASH::FStreamHandler::HandleRequest()
 			}
 			ProgressListener.Reset();
 			// Note: It is only safe to access the connection info when the HTTP request has completed or the request been removed.
-			PlayerSessionService->GetHTTPManager()->RemoveRequest(HTTP);
+			PlayerSessionService->GetHTTPManager()->RemoveRequest(HTTP, false);
 			CurrentRequest->ConnectionInfo = HTTP->ConnectionInfo;
 			HTTP.Reset();
 		}
diff --git a/Engine/Plugins/Media/ElectraPlayer/Source/ElectraPlayerRuntime/Private/Runtime/Player/DASH/StreamReaderFMP4DASH.h b/Engine/Plugins/Media/ElectraPlayer/Source/ElectraPlayerRuntime/Private/Runtime/Player/DASH/StreamReaderFMP4DASH.h
--- a/Engine/Plugins/Media/ElectraPlayer/Source/ElectraPlayerRuntime/Private/Runtime/Player/DASH/StreamReaderFMP4DASH.h
+++ b/Engine/Plugins/Media/ElectraPlayer/Source/ElectraPlayerRuntime/Private/Runtime/Player/DASH/StreamReaderFMP4DASH.h
@@ -8,6 +8,7 @@
 #include "Player/DASH/PlayerEventDASH.h"
 #include "Demuxer/ParserISO14496-12.h"
 #include "HTTP/HTTPManager.h"
+#include "Player/DRM/DRMManager.h"
 
 namespace Electra
 {
@@ -55,6 +56,11 @@ public:
 	FTimeValue												SAET;
 	FTimeValue												DownloadDelayTime;
 
+	// Encryption
+	TSharedPtrTS<ElectraCDM::IMediaCDMClient>				DrmClient;
+	FString													DrmMimeType;
+
+
 	// Internal work variables
 	FPlayerLoopState										PlayerLoopState;
 	int32													NumOverallRetries = 0;								//!< Number of retries for this _segment_ across all possible quality levels and CDNs.
diff --git a/Engine/Plugins/Media/ElectraPlayer/Source/ElectraPlayerRuntime/Private/Runtime/Player/DRM/DRMManager.cpp b/Engine/Plugins/Media/ElectraPlayer/Source/ElectraPlayerRuntime/Private/Runtime/Player/DRM/DRMManager.cpp
new file mode 100644
--- /dev/null
+++ b/Engine/Plugins/Media/ElectraPlayer/Source/ElectraPlayerRuntime/Private/Runtime/Player/DRM/DRMManager.cpp
@@ -0,0 +1,157 @@
+// Copyright Epic Games, Inc. All Rights Reserved.
+
+#include "Player/DRM/DRMManager.h"
+#include "Player/PlayerLicenseKey.h"
+
+namespace Electra
+{
+
+namespace
+{
+	ElectraCDM::IMediaCDM* Get()
+	{
+		return &ElectraCDM::IMediaCDM::Get();
+	}
+}
+
+FDRMManager::FDRMManager(IPlayerSessionServices* InPlayerSessionServices)
+	: PlayerSessionServices(InPlayerSessionServices)
+{
+	Session = Get()->CreatePlayerSessionID();
+}
+
+FDRMManager::~FDRMManager()
+{
+	if (Session)
+	{
+		Get()->ReleasePlayerSessionKeys(Session);
+		Get()->ReleasePlayerSessionID(Session);
+	}
+}
+
+TSharedPtrTS<FDRMManager> FDRMManager::Create(IPlayerSessionServices* InPlayerSessionServices)
+{
+	return MakeSharedTS<FDRMManager>(InPlayerSessionServices);
+}
+
+void FDRMManager::Tick()
+{
+	FScopeLock lock(&Lock);
+	CompletedRequests.Empty();
+}
+
+void FDRMManager::Close()
+{
+	FScopeLock lock(&Lock);
+	PendingRequests.Empty();
+	CompletedRequests.Empty();
+	if (Session)
+	{
+		Get()->ReleasePlayerSessionKeys(Session);
+		Get()->ReleasePlayerSessionID(Session);
+		Session = nullptr;
+	}
+}
+
+
+void FDRMManager::GetCDMCustomJSONPrefixes(const FString& InCDMScheme, const FString& InValue, FString& OutAttributePrefix, FString& OutTextPropertyName, bool& bOutNoNamespaces)
+{
+	Get()->GetCDMCustomJSONPrefixes(InCDMScheme, InValue, OutAttributePrefix, OutTextPropertyName, bOutNoNamespaces);
+}
+
+TSharedPtr<ElectraCDM::IMediaCDMCapabilities, ESPMode::ThreadSafe> FDRMManager::GetCDMCapabilitiesForScheme(const FString& InCDMScheme, const FString& InValue, const FString& InAdditionalElements)
+{
+	return Get()->GetCDMCapabilitiesForScheme(InCDMScheme, InValue, InAdditionalElements);
+}
+
+
+ElectraCDM::ECDMError FDRMManager::CreateDRMClient(TSharedPtr<ElectraCDM::IMediaCDMClient, ESPMode::ThreadSafe>& OutClient, const TArray<ElectraCDM::IMediaCDM::FCDMCandidate>& InCandidates)
+{
+	return Get()->CreateDRMClient(OutClient, Session, InCandidates);
+}
+
+void FDRMManager::OnHTTPResourceRequestComplete(TSharedPtrTS<FHTTPResourceRequest> InRequest)
+{
+	TSharedPtrTS<FDRMAsyncRequest> DrmReq = StaticCastSharedPtr<FDRMAsyncRequest>(InRequest->GetObject());
+	if (DrmReq.IsValid())
+	{
+		FScopeLock lock(&Lock);
+		PendingRequests.Remove(DrmReq);
+		CompletedRequests.Emplace(DrmReq);
+
+		TSharedPtr<ElectraCDM::IMediaCDMClient, ESPMode::ThreadSafe> Client = DrmReq->Client.Pin();
+		if (Client.IsValid())
+		{
+			TArray<uint8> LicenseResponse;
+
+			TSharedPtrTS<IElectraHttpManager::FReceiveBuffer> ResponseBuffer = DrmReq->Request->GetResponseBuffer();
+			if (ResponseBuffer.IsValid())
+			{
+				LicenseResponse.Append((const uint8*)ResponseBuffer->Buffer.GetLinearReadData(), ResponseBuffer->Buffer.Num());
+			}
+			const HTTP::FConnectionInfo* ConnInfo = DrmReq->Request->GetConnectionInfo();
+			int32 HttpResponseCode = ConnInfo ? ConnInfo->StatusInfo.HTTPStatus : 500;
+			PlayerSessionServices->SendMessageToPlayer(FLicenseKeyMessage::Create(FLicenseKeyMessage::EReason::LicenseKeyDownload, FErrorDetail(), ConnInfo));
+			if (Client->SetLicenseKeyResponseData(DrmReq->EventId, HttpResponseCode, LicenseResponse) != ElectraCDM::ECDMError::Success)
+			{
+				PlayerSessionServices->SendMessageToPlayer(FLicenseKeyMessage::Create(FLicenseKeyMessage::EReason::LicenseKeyData,
+																						FErrorDetail().SetError(UEMEDIA_ERROR_FORMAT_ERROR).SetFacility(Facility::EFacility::LicenseKey).SetMessage("Invalid license key"), ConnInfo));
+			}
+		}
+	}
+}
+
+
+void FDRMManager::OnCDMEvent(ElectraCDM::IMediaCDMEventListener::ECDMEventType InEventType, TSharedPtr<ElectraCDM::IMediaCDMClient, ESPMode::ThreadSafe> InDrmClient, void* InEventId, const TArray<uint8>& InCustomData)
+{
+	if (InEventType == ElectraCDM::IMediaCDMEventListener::ECDMEventType::KeyRequired)
+	{
+		FString LicenseURL, HttpMethod, MimeType;
+		TArray<uint8> LicenseData;
+		TArray<FString> Headers;
+		uint32 Flags = 0;
+		InDrmClient->GetLicenseKeyURL(LicenseURL);
+		InDrmClient->GetLicenseKeyRequestData(LicenseData, HttpMethod, Headers, Flags);
+		if (HttpMethod.IsEmpty())
+		{
+			HttpMethod = TEXT("POST");
+		}
+		if (MimeType.IsEmpty())
+		{
+			MimeType = TEXT("application/octet-stream");
+		}
+		if (LicenseURL.IsEmpty() && (Flags & ElectraCDM::IMediaCDMClient::EDRMClientFlags::EDRMFlg_AllowCustomKeyStorage) == 0)
+		{
+			PlayerSessionServices->PostError(FErrorDetail().SetFacility(Facility::EFacility::DRM).SetCode(1).SetError(UEMEDIA_ERROR_INTERNAL).SetMessage(FString::Printf(TEXT("No license URL specified to acquire DRM license from in event %d"), (int32)InEventType)));
+			return;
+		}
+
+		TSharedPtrTS<FDRMAsyncRequest> DrmReq = MakeSharedTS<FDRMAsyncRequest>();
+		DrmReq->Event = InEventType;
+		DrmReq->Client = InDrmClient;
+		DrmReq->EventId = InEventId;
+		DrmReq->CustomData = InCustomData;
+		DrmReq->Request = MakeSharedTS<FHTTPResourceRequest>();
+		DrmReq->Request->URL(LicenseURL).Verb(HttpMethod).Headers(Headers).PostData(LicenseData).Object(DrmReq);
+		DrmReq->Request->Callback().BindThreadSafeSP(AsShared(), &FDRMManager::OnHTTPResourceRequestComplete);
+		if ((Flags & ElectraCDM::IMediaCDMClient::EDRMClientFlags::EDRMFlg_AllowCustomKeyStorage) != 0)
+		{
+			DrmReq->Request->AllowStaticQuery(IAdaptiveStreamingPlayerResourceRequest::EPlaybackResourceType::LicenseKey);
+		}
+		Lock.Lock();
+		PendingRequests.Emplace(DrmReq);
+		Lock.Unlock();
+		DrmReq->Request->StartGet(PlayerSessionServices);
+	}
+	else
+	{
+		PlayerSessionServices->PostError(FErrorDetail().SetFacility(Facility::EFacility::DRM).SetCode(1).SetError(UEMEDIA_ERROR_INTERNAL).SetMessage(FString::Printf(TEXT("Unhandled DRM event %d"), (int32)InEventType)));
+	}
+}
+
+
+} // namespace Electra
+
+
+
+
diff --git a/Engine/Plugins/Media/ElectraPlayer/Source/ElectraPlayerRuntime/Private/Runtime/Player/DRM/DRMManager.h b/Engine/Plugins/Media/ElectraPlayer/Source/ElectraPlayerRuntime/Private/Runtime/Player/DRM/DRMManager.h
new file mode 100644
--- /dev/null
+++ b/Engine/Plugins/Media/ElectraPlayer/Source/ElectraPlayerRuntime/Private/Runtime/Player/DRM/DRMManager.h
@@ -0,0 +1,55 @@
+// Copyright Epic Games, Inc. All Rights Reserved.
+
+#pragma once
+
+#include "PlayerCore.h"
+#include "ErrorDetail.h"
+#include "Player/PlayerSessionServices.h"
+#include "Player/AdaptiveStreamingPlayerResourceRequest.h"
+#include "ElectraCDM.h"
+#include "ElectraCDMClient.h"
+
+
+namespace Electra
+{
+
+class FDRMManager : public TSharedFromThis<FDRMManager, ESPMode::ThreadSafe>, public ElectraCDM::IMediaCDMEventListener
+{
+public:
+	static TSharedPtrTS<FDRMManager> Create(IPlayerSessionServices* InPlayerSessionServices);
+	void Tick();
+	void Close();
+
+	FDRMManager(IPlayerSessionServices* InPlayerSessionServices);
+	virtual ~FDRMManager();
+
+	void GetCDMCustomJSONPrefixes(const FString& InCDMScheme, const FString& InValue, FString& OutAttributePrefix, FString& OutTextPropertyName, bool& bOutNoNamespaces);
+	TSharedPtr<ElectraCDM::IMediaCDMCapabilities, ESPMode::ThreadSafe> GetCDMCapabilitiesForScheme(const FString& InCDMScheme, const FString& InValue, const FString& InAdditionalElements);
+
+	ElectraCDM::ECDMError CreateDRMClient(TSharedPtr<ElectraCDM::IMediaCDMClient, ESPMode::ThreadSafe>& OutClient, const TArray<ElectraCDM::IMediaCDM::FCDMCandidate>& InCandidates);
+
+	virtual void OnCDMEvent(ElectraCDM::IMediaCDMEventListener::ECDMEventType InEventType, TSharedPtr<ElectraCDM::IMediaCDMClient, ESPMode::ThreadSafe> InDrmClient, void* InEventId, const TArray<uint8>& InCustomData) override;
+
+private:
+
+	struct FDRMAsyncRequest : public IHTTPResourceRequestObject
+	{
+		ElectraCDM::IMediaCDMEventListener::ECDMEventType Event;
+		TWeakPtr<ElectraCDM::IMediaCDMClient, ESPMode::ThreadSafe> Client;
+		TArray<uint8> CustomData;
+		void* EventId = nullptr;
+		TSharedPtrTS<FHTTPResourceRequest> Request;
+	};
+
+	void OnHTTPResourceRequestComplete(TSharedPtrTS<FHTTPResourceRequest> InRequest);
+
+	FCriticalSection Lock;
+	IPlayerSessionServices* PlayerSessionServices = nullptr;
+	ElectraCDM::IMediaCDM::IPlayerSession* Session = nullptr;
+	TArray<TSharedPtrTS<FDRMAsyncRequest>> PendingRequests;
+	TArray<TSharedPtrTS<FDRMAsyncRequest>> CompletedRequests;
+};
+
+
+} // namespace Electra
+
diff --git a/Engine/Plugins/Media/ElectraPlayer/Source/ElectraPlayerRuntime/Private/Runtime/Player/HLS/PlaylistReaderHLS.cpp b/Engine/Plugins/Media/ElectraPlayer/Source/ElectraPlayerRuntime/Private/Runtime/Player/HLS/PlaylistReaderHLS.cpp
--- a/Engine/Plugins/Media/ElectraPlayer/Source/ElectraPlayerRuntime/Private/Runtime/Player/HLS/PlaylistReaderHLS.cpp
+++ b/Engine/Plugins/Media/ElectraPlayer/Source/ElectraPlayerRuntime/Private/Runtime/Player/HLS/PlaylistReaderHLS.cpp
@@ -635,7 +635,7 @@ void FPlaylistReaderHLS::FPlaylistRequest::Execute(TSharedPtrTS<IElectraHttpMana
 	HTTPRequest->ReceiveBuffer    = ReceiveBuffer;
 	HTTPRequest->ProgressListener = InProgressListener;
 	// Add the request.
-	InPlayerSessionServices->GetHTTPManager()->AddRequest(HTTPRequest);
+	InPlayerSessionServices->GetHTTPManager()->AddRequest(HTTPRequest, false);
 }
 
 void FPlaylistReaderHLS::FPlaylistRequest::Cancel(IPlayerSessionServices* InPlayerSessionServices)
@@ -647,7 +647,7 @@ void FPlaylistReaderHLS::FPlaylistRequest::Cancel(IPlayerSessionServices* InPlay
 	}
 	if (HTTPRequest.IsValid())
 	{
-		InPlayerSessionServices->GetHTTPManager()->RemoveRequest(HTTPRequest);
+		InPlayerSessionServices->GetHTTPManager()->RemoveRequest(HTTPRequest, false);
 		HTTPRequest.Reset();
 	}
 }
diff --git a/Engine/Plugins/Media/ElectraPlayer/Source/ElectraPlayerRuntime/Private/Runtime/Player/HLS/StreamReaderHLSfmp4.cpp b/Engine/Plugins/Media/ElectraPlayer/Source/ElectraPlayerRuntime/Private/Runtime/Player/HLS/StreamReaderHLSfmp4.cpp
--- a/Engine/Plugins/Media/ElectraPlayer/Source/ElectraPlayerRuntime/Private/Runtime/Player/HLS/StreamReaderHLSfmp4.cpp
+++ b/Engine/Plugins/Media/ElectraPlayer/Source/ElectraPlayerRuntime/Private/Runtime/Player/HLS/StreamReaderHLSfmp4.cpp
@@ -12,18 +12,11 @@
 #define INTERNAL_ERROR_INIT_SEGMENT_PARSE_ERROR						2
 #define INTERNAL_ERROR_INIT_SEGMENT_LICENSEKEY_ERROR				10
 
-#define PERMIT_INSECURE_SEGMENT_DEMOTING 0
+
 
 namespace Electra
 {
 
-//! (bool) if false and media segment is using EXT-X-KEY encryption fetch it via http even if it should be https, otherwise keep the original scheme.
-const FString FStreamReaderHLSfmp4::OptionKeyDontUseInsecureForEncryptedMediaSegments("dont_use_insecure_for_media_segments");
-
-//! (bool) if false the init segment is fetched via http even if it should be https, otherwise keep the original scheme.
-const FString FStreamReaderHLSfmp4::OptionKeyDontUseInsecureForInitSegments("dont_use_insecure_for_init_segments");
-
-
 FStreamSegmentRequestHLSfmp4::FStreamSegmentRequestHLSfmp4()
 {
 	StreamType  			  = EStreamType::Video;
@@ -396,35 +389,6 @@ void FStreamReaderHLSfmp4::FStreamHandler::HTTPUpdateStats(const FTimeValue& Cur
 }
 
 
-FString FStreamReaderHLSfmp4::FStreamHandler::DemoteInitURLToHTTP(const FString& InURL, bool bIsEncrypted)
-{
-	FString NewUrl(InURL);
-#if PERMIT_INSECURE_SEGMENT_DEMOTING
-	if (bIsEncrypted &&
-		Parameters.Options.GetValue(OptionKeyDontUseInsecureForInitSegments).SafeGetBool(false) == false &&
-		NewUrl.Left(8) == TEXT("https://"))
-	{
-		// Remove the 's'
-		NewUrl.RemoveAt(4);
-	}
-#endif
-	return NewUrl;
-}
-
-FString FStreamReaderHLSfmp4::FStreamHandler::DemoteMediaURLToHTTP(const FString& InURL, bool bIsEncrypted)
-{
-	FString NewUrl(InURL);
-#if PERMIT_INSECURE_SEGMENT_DEMOTING
-if (bIsEncrypted &&
-		Parameters.Options.GetValue(OptionKeyDontUseInsecureForEncryptedMediaSegments).SafeGetBool(false) == false &&
-		NewUrl.Left(8) == TEXT("https://"))
-	{
-		// Remove the 's'
-		NewUrl.RemoveAt(4);
-	}
-#endif
-	return NewUrl;
-}
 
 
 FStreamReaderHLSfmp4::FStreamHandler::ELicenseKeyResult FStreamReaderHLSfmp4::FStreamHandler::GetLicenseKey(FErrorDetail& OutErrorDetail, TSharedPtr<TArray<uint8>, ESPMode::ThreadSafe>& OutLicenseKeyData, const TSharedPtrTS<FStreamSegmentRequestHLSfmp4>& InRequest, const TSharedPtr<const FManifestHLSInternal::FMediaStream::FDRMKeyInfo, ESPMode::ThreadSafe>& LicenseKeyInfo)
@@ -488,7 +452,7 @@ FStreamReaderHLSfmp4::FStreamHandler::ELicenseKeyResult FStreamReaderHLSfmp4::FS
 				ProgressListener->ProgressDelegate   = Electra::MakeDelegate(this, &FStreamHandler::HTTPProgressCallback);
 				ProgressListener->CompletionDelegate = Electra::MakeDelegate(this, &FStreamHandler::HTTPCompletionCallback);
 				HTTP->ProgressListener = ProgressListener;
-				PlayerSessionService->GetHTTPManager()->AddRequest(HTTP);
+				PlayerSessionService->GetHTTPManager()->AddRequest(HTTP, false);
 				while(!HasReadBeenAborted())
 				{
 					if (DownloadCompleteSignal.WaitTimeout(1000 * 100))
@@ -498,7 +462,7 @@ FStreamReaderHLSfmp4::FStreamHandler::ELicenseKeyResult FStreamReaderHLSfmp4::FS
 				}
 				ProgressListener.Reset();
 				// Note: It is only safe to access the connection info when the HTTP request has completed or the request been removed.
-				PlayerSessionService->GetHTTPManager()->RemoveRequest(HTTP);
+				PlayerSessionService->GetHTTPManager()->RemoveRequest(HTTP, false);
 			}
 
 			InRequest->ConnectionInfo = HTTP->ConnectionInfo;
@@ -572,10 +536,7 @@ FStreamReaderHLSfmp4::FStreamHandler::EInitSegmentResult FStreamReaderHLSfmp4::F
 			ReadBuffer.Reset();
 			ReadBuffer.ReceiveBuffer = MakeSharedTS<IElectraHttpManager::FReceiveBuffer>();
 
-			// FIXME: Presently we demote init segments to HTTP unless explicitly forbidden to do so.
-			//        Since they are not encrypted we have no reliable way to know if the content itself is encrypted.
-			//        We say it is if _any_ segment in the playlist is using encryption, whichever one that is.
-			FString RequestURL = DemoteInitURLToHTTP(Request->InitSegmentInfo->URI, Request->bHasEncryptedSegments);
+			FString RequestURL = Request->InitSegmentInfo->URI;
 
 			Metrics::FSegmentDownloadStats& ds = Request->DownloadStats;
 			ds.URL  	   = RequestURL;
@@ -592,7 +553,7 @@ FStreamReaderHLSfmp4::FStreamHandler::EInitSegmentResult FStreamReaderHLSfmp4::F
 
 			ProgressReportCount = 0;
 			DownloadCompleteSignal.Reset();
-			PlayerSessionService->GetHTTPManager()->AddRequest(HTTP);
+			PlayerSessionService->GetHTTPManager()->AddRequest(HTTP, false);
 
 			while(!HasReadBeenAborted())
 			{
@@ -604,7 +565,7 @@ FStreamReaderHLSfmp4::FStreamHandler::EInitSegmentResult FStreamReaderHLSfmp4::F
 
 			ProgressListener.Reset();
 			// Note: It is only safe to access the connection info when the HTTP request has completed or the request been removed.
-			PlayerSessionService->GetHTTPManager()->RemoveRequest(HTTP);
+			PlayerSessionService->GetHTTPManager()->RemoveRequest(HTTP, false);
 			Request->ConnectionInfo = HTTP->ConnectionInfo;
 			if (!HTTP->ConnectionInfo.StatusInfo.ErrorDetail.IsError())
 			{
@@ -789,15 +750,15 @@ void FStreamReaderHLSfmp4::FStreamHandler::HandleRequest()
 				{
 					if (LicenseKeyData.IsValid())
 					{
-						IStreamDecrypterAES128::EResult DecrypterResult;
-						Decrypter = IStreamDecrypterAES128::Create();
+						ElectraCDM::IStreamDecrypterAES128::EResult DecrypterResult;
+						Decrypter = ElectraCDM::IStreamDecrypterAES128::Create();
 
 						// Set up the IV for this segment which is either explicitly provided or the media sequence number.
 						TArray<uint8> IV;
 						if (Request->LicenseKeyInfo->IV.Len())
 						{
-							DecrypterResult = IStreamDecrypterAES128::ConvHexStringToBin(IV, TCHAR_TO_ANSI(*Request->LicenseKeyInfo->IV));
-							if (DecrypterResult != IStreamDecrypterAES128::EResult::Ok)
+							DecrypterResult = ElectraCDM::IStreamDecrypterAES128::ConvHexStringToBin(IV, TCHAR_TO_ANSI(*Request->LicenseKeyInfo->IV));
+							if (DecrypterResult != ElectraCDM::IStreamDecrypterAES128::EResult::Ok)
 							{
 								ds.FailureReason = FString::Printf(TEXT("Bad explicit IV value"));
 								LogMessage(IInfoLog::ELevel::Error, ds.FailureReason);
@@ -806,12 +767,12 @@ void FStreamReaderHLSfmp4::FStreamHandler::HandleRequest()
 						}
 						else
 						{
-							IStreamDecrypterAES128::MakePaddedIVFromUInt64(IV, Request->MediaSequence);
+							ElectraCDM::IStreamDecrypterAES128::MakePaddedIVFromUInt64(IV, Request->MediaSequence);
 						}
 						if (!bHasErrored)
 						{
 							DecrypterResult = Decrypter->CBCInit(*LicenseKeyData, &IV);
-							if (DecrypterResult != IStreamDecrypterAES128::EResult::Ok)
+							if (DecrypterResult != ElectraCDM::IStreamDecrypterAES128::EResult::Ok)
 							{
 								ds.FailureReason = FString::Printf(TEXT("Received bad license key"));
 								LogMessage(IInfoLog::ELevel::Error, ds.FailureReason);
@@ -834,7 +795,7 @@ void FStreamReaderHLSfmp4::FStreamHandler::HandleRequest()
 				ReadBuffer.ReceiveBuffer = MakeSharedTS<IElectraHttpManager::FReceiveBuffer>();
 
 				// Start downloading the segment. Clear any stats that may have been set by the init segment download.
-				FString RequestURL = DemoteMediaURLToHTTP(Request->URL, Decrypter.IsValid());
+				FString RequestURL = Request->URL;
 
 				ds.FailureReason.Empty();
 				ds.URL  			   = RequestURL;
@@ -866,7 +827,7 @@ void FStreamReaderHLSfmp4::FStreamHandler::HandleRequest()
 
 				ProgressReportCount = 0;
 				DownloadCompleteSignal.Reset();
-				PlayerSessionService->GetHTTPManager()->AddRequest(HTTP);
+				PlayerSessionService->GetHTTPManager()->AddRequest(HTTP, false);
 
 				MP4Parser = IParserISO14496_12::CreateParser();
 				NumMOOFBoxesFound = 0;
@@ -1051,7 +1012,7 @@ void FStreamReaderHLSfmp4::FStreamHandler::HandleRequest()
 				}
 				ProgressListener.Reset();
 				// Note: It is only safe to access the connection info when the HTTP request has completed or the request been removed.
-				PlayerSessionService->GetHTTPManager()->RemoveRequest(HTTP);
+				PlayerSessionService->GetHTTPManager()->RemoveRequest(HTTP, false);
 				CurrentRequest->ConnectionInfo = HTTP->ConnectionInfo;
 				HTTP.Reset();
 			}
@@ -1382,11 +1343,11 @@ int64 FStreamReaderHLSfmp4::FStreamHandler::ReadData(void* IntoBuffer, int64 Num
 						}
 
 						// Decrypt data in place
-						IStreamDecrypterAES128::EResult DecrypterResult;
+						ElectraCDM::IStreamDecrypterAES128::EResult DecrypterResult;
 						int32 NumDecryptedBytes = 0;
 						DecrypterResult = Decrypter->CBCDecryptInPlace(NumDecryptedBytes, EncryptedData, EncryptedSize, bIsFinalBlock);
 						// This cannot fail since we ensured it to be set up correctly and pass only properly aligned data, but just in case.
-						if (DecrypterResult == IStreamDecrypterAES128::EResult::Ok)
+						if (DecrypterResult == ElectraCDM::IStreamDecrypterAES128::EResult::Ok)
 						{
 							// Advance the decrypted pos by the entire encrypted block size, not the amount of decrypted bytes!
 							// This is required for the RequiredEncryptedSize test above to work correctly.
@@ -1402,7 +1363,7 @@ int64 FStreamReaderHLSfmp4::FStreamHandler::ReadData(void* IntoBuffer, int64 Num
 						else
 						{
 							SourceBuffer.Unlock();
-							LogMessage(IInfoLog::ELevel::Error, FString::Printf(TEXT("Failed to decrypt (%s)"), IStreamDecrypterAES128::GetResultText(DecrypterResult)));
+							LogMessage(IInfoLog::ELevel::Error, FString::Printf(TEXT("Failed to decrypt (%s)"), ElectraCDM::IStreamDecrypterAES128::GetResultText(DecrypterResult)));
 							return -1;
 						}
 					}
diff --git a/Engine/Plugins/Media/ElectraPlayer/Source/ElectraPlayerRuntime/Private/Runtime/Player/HLS/StreamReaderHLSfmp4.h b/Engine/Plugins/Media/ElectraPlayer/Source/ElectraPlayerRuntime/Private/Runtime/Player/HLS/StreamReaderHLSfmp4.h
--- a/Engine/Plugins/Media/ElectraPlayer/Source/ElectraPlayerRuntime/Private/Runtime/Player/HLS/StreamReaderHLSfmp4.h
+++ b/Engine/Plugins/Media/ElectraPlayer/Source/ElectraPlayerRuntime/Private/Runtime/Player/HLS/StreamReaderHLSfmp4.h
@@ -200,7 +200,7 @@ private:
 
 		IPlayerSessionServices*									PlayerSessionService;
 		FReadBuffer												ReadBuffer;
-		TSharedPtr<IStreamDecrypterAES128, ESPMode::ThreadSafe>	Decrypter;
+		TSharedPtr<ElectraCDM::IStreamDecrypterAES128, ESPMode::ThreadSafe>	Decrypter;
 		FMediaEvent												DownloadCompleteSignal;
 		TSharedPtrTS<IParserISO14496_12>						MP4Parser;
 		int32													NumMOOFBoxesFound;
@@ -220,8 +220,6 @@ private:
 		void SignalWork();
 		void WorkerThread();
 		void HandleRequest();
-		FString DemoteMediaURLToHTTP(const FString& InURL, bool bIsEncrypted);
-		FString DemoteInitURLToHTTP(const FString& InURL, bool bIsEncrypted);
 		EInitSegmentResult GetInitSegment(FErrorDetail& OutErrorDetail, TSharedPtrTS<const IParserISO14496_12>& OutMP4InitSegment, const TSharedPtrTS<FStreamSegmentRequestHLSfmp4>& InRequest);
 		ELicenseKeyResult GetLicenseKey(FErrorDetail& OutErrorDetail, TSharedPtr<TArray<uint8>, ESPMode::ThreadSafe>& OutLicenseKeyData, const TSharedPtrTS<FStreamSegmentRequestHLSfmp4>& InRequest, const TSharedPtr<const FManifestHLSInternal::FMediaStream::FDRMKeyInfo, ESPMode::ThreadSafe>& LicenseKeyInfo);
 
@@ -250,9 +248,6 @@ private:
 	IPlayerSessionServices*				PlayerSessionService;
 	bool								bIsStarted;
 	FErrorDetail						ErrorDetail;
-
-	static const FString		OptionKeyDontUseInsecureForEncryptedMediaSegments;		//!< (bool) if false and media segment is using EXT-X-KEY encryption fetch it via http even if it should be https, otherwise keep the original scheme.
-	static const FString		OptionKeyDontUseInsecureForInitSegments;				//!< (bool) if false the init segment is fetched via http even if it should be https, otherwise keep the original scheme.
 };
 
 
diff --git a/Engine/Plugins/Media/ElectraPlayer/Source/ElectraPlayerRuntime/Private/Runtime/Player/PlayerSessionServices.h b/Engine/Plugins/Media/ElectraPlayer/Source/ElectraPlayerRuntime/Private/Runtime/Player/PlayerSessionServices.h
--- a/Engine/Plugins/Media/ElectraPlayer/Source/ElectraPlayerRuntime/Private/Runtime/Player/PlayerSessionServices.h
+++ b/Engine/Plugins/Media/ElectraPlayer/Source/ElectraPlayerRuntime/Private/Runtime/Player/PlayerSessionServices.h
@@ -7,7 +7,6 @@
 #include "InfoLog.h"
 #include "StreamTypes.h"
 
-
 namespace Electra
 {
 	class ISynchronizedUTCTime;
@@ -19,6 +18,7 @@ namespace Electra
 	class IPlayerEntityCache;
 	class IPlaylistReader;
 	class IAdaptiveStreamingPlayerAEMSHandler;
+	class FDRMManager;
 
 
 	class IPlayerMessage
@@ -126,6 +126,11 @@ namespace Electra
 		 * pass values between internal sub systems.
 		 */
 		 virtual FParamDict& GetOptions() = 0;
+
+		 /**
+		  * Returns the DRM manager, if any.
+		  */
+		virtual TSharedPtrTS<FDRMManager> GetDRMManager() = 0;
 	};
 
 
diff --git a/Engine/Plugins/Media/ElectraPlayer/Source/ElectraPlayerRuntime/Private/Runtime/Player/mp4/PlaylistReaderMP4.cpp b/Engine/Plugins/Media/ElectraPlayer/Source/ElectraPlayerRuntime/Private/Runtime/Player/mp4/PlaylistReaderMP4.cpp
--- a/Engine/Plugins/Media/ElectraPlayer/Source/ElectraPlayerRuntime/Private/Runtime/Player/mp4/PlaylistReaderMP4.cpp
+++ b/Engine/Plugins/Media/ElectraPlayer/Source/ElectraPlayerRuntime/Private/Runtime/Player/mp4/PlaylistReaderMP4.cpp
@@ -288,7 +288,7 @@ void FPlaylistReaderMP4::ClearRequest()
 	ReceiveBuffer.Reset();
 	if (Request.IsValid())
 	{
-		PlayerSessionServices->GetHTTPManager()->RemoveRequest(Request);
+		PlayerSessionServices->GetHTTPManager()->RemoveRequest(Request, false);
 		Request.Reset();
 	}
 }
@@ -323,7 +323,7 @@ void FPlaylistReaderMP4::ReadNextChunk(int64 InFromOffset, int64 ChunkSize)
 	Request->Parameters.Range.SetEndIncluding(LastByte);
 	Request->ReceiveBuffer = ReceiveBuffer;
 	Request->ProgressListener = ProgressListener;
-	PlayerSessionServices->GetHTTPManager()->AddRequest(Request);
+	PlayerSessionServices->GetHTTPManager()->AddRequest(Request, false);
 }
 
 
diff --git a/Engine/Plugins/Media/ElectraPlayer/Source/ElectraPlayerRuntime/Private/Runtime/Player/mp4/StreamReaderMP4.cpp b/Engine/Plugins/Media/ElectraPlayer/Source/ElectraPlayerRuntime/Private/Runtime/Player/mp4/StreamReaderMP4.cpp
--- a/Engine/Plugins/Media/ElectraPlayer/Source/ElectraPlayerRuntime/Private/Runtime/Player/mp4/StreamReaderMP4.cpp
+++ b/Engine/Plugins/Media/ElectraPlayer/Source/ElectraPlayerRuntime/Private/Runtime/Player/mp4/StreamReaderMP4.cpp
@@ -400,7 +400,7 @@ void FStreamReaderMP4::WorkerThread()
 
 		HTTP->ReceiveBuffer 				= ReadBuffer.ReceiveBuffer;
 		HTTP->ProgressListener  			= ProgressListener;
-		PlayerSessionServices->GetHTTPManager()->AddRequest(HTTP);
+		PlayerSessionServices->GetHTTPManager()->AddRequest(HTTP, false);
 
 
 		FTimeValue DurationSuccessfullyDelivered(FTimeValue::GetZero());
@@ -619,7 +619,7 @@ void FStreamReaderMP4::WorkerThread()
 
 		// Remove the download request.
 		ProgressListener.Reset();
-		PlayerSessionServices->GetHTTPManager()->RemoveRequest(HTTP);
+		PlayerSessionServices->GetHTTPManager()->RemoveRequest(HTTP, false);
 		Request->ConnectionInfo = HTTP->ConnectionInfo;
 		HTTP.Reset();
 
diff --git a/Engine/Plugins/Media/ElectraPlayer/Source/ElectraPlayerRuntime/Private/Runtime/PlayerFacility.h b/Engine/Plugins/Media/ElectraPlayer/Source/ElectraPlayerRuntime/Private/Runtime/PlayerFacility.h
--- a/Engine/Plugins/Media/ElectraPlayer/Source/ElectraPlayerRuntime/Private/Runtime/PlayerFacility.h
+++ b/Engine/Plugins/Media/ElectraPlayer/Source/ElectraPlayerRuntime/Private/Runtime/PlayerFacility.h
@@ -39,6 +39,7 @@ namespace Electra
 			DASHManifest,
 			DASHFMP4Reader,
 			EntityCache,
+			DRM,
 			LastEntry = 255
 		};
 
@@ -98,6 +99,8 @@ namespace Electra
 					return TEXT("DASH fmp4 reader");
 				case EFacility::EntityCache:
 					return TEXT("Entity cache");
+				case EFacility::DRM:
+					return TEXT("DRM");
 			}
 			return TEXT("???");
 		}
diff --git a/Engine/Plugins/Media/ElectraPlayer/Source/ElectraPlayerRuntime/Private/Runtime/StreamTypes.cpp b/Engine/Plugins/Media/ElectraPlayer/Source/ElectraPlayerRuntime/Private/Runtime/StreamTypes.cpp
--- a/Engine/Plugins/Media/ElectraPlayer/Source/ElectraPlayerRuntime/Private/Runtime/StreamTypes.cpp
+++ b/Engine/Plugins/Media/ElectraPlayer/Source/ElectraPlayerRuntime/Private/Runtime/StreamTypes.cpp
@@ -7,6 +7,40 @@
 namespace Electra
 {
 
+	FString FStreamCodecInformation::GetMimeType() const
+	{
+		switch(GetCodec())
+		{
+			case ECodec::H264:
+				return FString(TEXT("video/mp4"));
+			case ECodec::H265:
+				return FString(TEXT("video/mp4"));
+			case ECodec::AAC:
+				return FString(TEXT("audio/mp4"));
+			case ECodec::EAC3:
+				return FString(TEXT("audio/mp4"));
+			case ECodec::WebVTT:
+				return FString(TEXT("application/mp4"));
+			default:
+				return FString(TEXT("application/octet-stream"));
+		}
+	}
+
+	FString FStreamCodecInformation::GetMimeTypeWithCodec() const
+	{
+		return GetMimeType() + FString::Printf(TEXT("; codecs=\"%s\""), *GetCodecSpecifierRFC6381());
+	}
+
+	FString FStreamCodecInformation::GetMimeTypeWithCodecAndFeatures() const
+	{
+		if (GetStreamType() == EStreamType::Video && GetResolution().Width && GetResolution().Height)
+		{
+			return GetMimeTypeWithCodec() + FString::Printf(TEXT("; resolution=%dx%d"), GetResolution().Width, GetResolution().Height);
+		}
+		return GetMimeTypeWithCodec();
+	}
+
+
 	bool FStreamCodecInformation::ParseFromRFC6381(const FString& CodecOTI)
 	{
 		if (CodecOTI.StartsWith("avc"))
diff --git a/Engine/Plugins/Media/ElectraPlayer/Source/ElectraPlayerRuntime/Private/Runtime/StreamTypes.h b/Engine/Plugins/Media/ElectraPlayer/Source/ElectraPlayerRuntime/Private/Runtime/StreamTypes.h
--- a/Engine/Plugins/Media/ElectraPlayer/Source/ElectraPlayerRuntime/Private/Runtime/StreamTypes.h
+++ b/Engine/Plugins/Media/ElectraPlayer/Source/ElectraPlayerRuntime/Private/Runtime/StreamTypes.h
@@ -61,6 +61,10 @@ namespace Electra
 			return StreamType;
 		}
 
+		FString GetMimeType() const;
+		FString GetMimeTypeWithCodec() const;
+		FString GetMimeTypeWithCodecAndFeatures() const;
+
 		bool ParseFromRFC6381(const FString& CodecOTI);
 
 		void SetStreamType(EStreamType InStreamType)
