Added missing include

diff --git a/Engine/Plugins/Media/ElectraCDM/Source/ElectraCDM/Private/ElectraCDM.cpp b/Engine/Plugins/Media/ElectraCDM/Source/ElectraCDM/Private/ElectraCDM.cpp
--- a/Engine/Plugins/Media/ElectraCDM/Source/ElectraCDM/Private/ElectraCDM.cpp
+++ b/Engine/Plugins/Media/ElectraCDM/Source/ElectraCDM/Private/ElectraCDM.cpp
@@ -1,6 +1,7 @@
 // Copyright Epic Games, Inc. All Rights Reserved.
 
 #include "ElectraCDM.h"
+#include "ElectraCDMSystem.h"
 //#include "ClearKey/ClearKeyCDM.h"
 #include "HAL/CriticalSection.h"
 #include "Misc/ScopeLock.h"
