Electra: Adjusted version number and removing beta state from ElectraPlayer and associated plugins

#rnx

diff --git a/Engine/Plugins/Media/ElectraCDM/ElectraCDM.uplugin b/Engine/Plugins/Media/ElectraCDM/ElectraCDM.uplugin
--- a/Engine/Plugins/Media/ElectraCDM/ElectraCDM.uplugin
+++ b/Engine/Plugins/Media/ElectraCDM/ElectraCDM.uplugin
@@ -1,7 +1,7 @@
 {
 	"FileVersion" : 1,
-	"Version" : 0,
-	"VersionName" : "0.1",
+	"Version" : 1,
+	"VersionName" : "1.0",
 	"FriendlyName" : "Electra Player Content Decryption Module",
 	"Description" : "Content Decryption Module for Electra Player Media Playback",
 	"Category" : "Media Players",
@@ -12,7 +12,7 @@
 	"SupportURL" : "",
 	"EnabledByDefault" : false,
 	"CanContainContent" : false,
-	"IsBetaVersion" : true,
+	"IsBetaVersion" : false,
 	"Installed" : false,
 	"Modules": [
 		{
diff --git a/Engine/Plugins/Media/ElectraPlayer/ElectraPlayer.uplugin b/Engine/Plugins/Media/ElectraPlayer/ElectraPlayer.uplugin
--- a/Engine/Plugins/Media/ElectraPlayer/ElectraPlayer.uplugin
+++ b/Engine/Plugins/Media/ElectraPlayer/ElectraPlayer.uplugin
@@ -1,7 +1,7 @@
 {
 	"FileVersion": 1,
-	"Version": 0,
-	"VersionName": "0.1",
+	"Version": 1,
+	"VersionName": "1.0",
 	"FriendlyName": "Electra Player",
 	"Description": "Next-generation playback capability.",
 	"Category": "Media Players",
@@ -12,7 +12,7 @@
 	"SupportURL": "",
 	"EnabledByDefault": false,
 	"CanContainContent": false,
-	"IsBetaVersion": true,
+	"IsBetaVersion": false,
 	"Installed": false,
 	"Modules": [
 		{
diff --git a/Engine/Plugins/Media/ElectraUtil/ElectraUtil.uplugin b/Engine/Plugins/Media/ElectraUtil/ElectraUtil.uplugin
--- a/Engine/Plugins/Media/ElectraUtil/ElectraUtil.uplugin
+++ b/Engine/Plugins/Media/ElectraUtil/ElectraUtil.uplugin
@@ -1,7 +1,7 @@
 {
 	"FileVersion" : 1,
-	"Version" : 0,
-	"VersionName" : "0.1",
+	"Version" : 1,
+	"VersionName" : "1.0",
 	"FriendlyName" : "Electra Player Utilities",
 	"Description" : "Reusable Base Components for Electra Player Media Playback",
 	"Category" : "Media Players",
@@ -12,7 +12,7 @@
 	"SupportURL" : "",
 	"EnabledByDefault" : false,
 	"CanContainContent" : false,
-	"IsBetaVersion" : true,
+	"IsBetaVersion" : false,
 	"Installed" : false,
 	"Modules": [
 		{
