Updated several modules to use shared PCHs.

Compiling only ElectraCDM module:
Before: 3.28 seconds
After: 2.22 seconds

Compiling only ElectraPlayerPlugin module:
Before: 4.42 seconds
After: 2.22 seconds

Compiling only ElectraSamples module:
Before: 5.50 seconds
After: 4.04 seconds

diff --git a/Engine/Plugins/Media/ElectraCDM/Source/ElectraCDM/ElectraCDM.Build.cs b/Engine/Plugins/Media/ElectraCDM/Source/ElectraCDM/ElectraCDM.Build.cs
--- a/Engine/Plugins/Media/ElectraCDM/Source/ElectraCDM/ElectraCDM.Build.cs
+++ b/Engine/Plugins/Media/ElectraCDM/Source/ElectraCDM/ElectraCDM.Build.cs
@@ -6,7 +6,8 @@ namespace UnrealBuildTool.Rules
 	{
 		public ElectraCDM(ReadOnlyTargetRules Target) : base(Target)
 		{
-			PCHUsage = PCHUsageMode.NoPCHs;     // to get around the "xzy must be first header included" error in 3rdParty lib
+			PCHUsage = PCHUsageMode.UseSharedPCHs;
+			PrivatePCHHeaderFile = "Public/ElectraCDM.h";
 			PublicDependencyModuleNames.AddRange(
 				new string[]
 				{
diff --git a/Engine/Plugins/Media/ElectraPlayer/Source/ElectraPlayerPlugin/ElectraPlayerPlugin.Build.cs b/Engine/Plugins/Media/ElectraPlayer/Source/ElectraPlayerPlugin/ElectraPlayerPlugin.Build.cs
--- a/Engine/Plugins/Media/ElectraPlayer/Source/ElectraPlayerPlugin/ElectraPlayerPlugin.Build.cs
+++ b/Engine/Plugins/Media/ElectraPlayer/Source/ElectraPlayerPlugin/ElectraPlayerPlugin.Build.cs
@@ -8,7 +8,6 @@ namespace UnrealBuildTool.Rules
 		public ElectraPlayerPlugin(ReadOnlyTargetRules Target) : base(Target)
 		{
 			bLegalToDistributeObjectCode = true;
-			PCHUsage = PCHUsageMode.NoPCHs;
 
 			DynamicallyLoadedModuleNames.AddRange(
 				new string[] {
diff --git a/Engine/Plugins/Media/ElectraUtil/Source/ElectraSamples/ElectraSamples.Build.cs b/Engine/Plugins/Media/ElectraUtil/Source/ElectraSamples/ElectraSamples.Build.cs
--- a/Engine/Plugins/Media/ElectraUtil/Source/ElectraSamples/ElectraSamples.Build.cs
+++ b/Engine/Plugins/Media/ElectraUtil/Source/ElectraSamples/ElectraSamples.Build.cs
@@ -8,9 +8,7 @@ namespace UnrealBuildTool.Rules
 	{
 		public ElectraSamples(ReadOnlyTargetRules Target) : base(Target)
 		{
-
 			bLegalToDistributeObjectCode = true;
-			PCHUsage = PCHUsageMode.NoPCHs;
 
 			DynamicallyLoadedModuleNames.AddRange(
 				new string[] {
