ElectraCDM: Added support for 'cbcs' scheme

#rnx

diff --git a/Engine/Plugins/Media/ElectraCDM/Source/ElectraCDM/Private/ClearKey/ClearKeyCDM.cpp b/Engine/Plugins/Media/ElectraCDM/Source/ElectraCDM/Private/ClearKey/ClearKeyCDM.cpp
--- a/Engine/Plugins/Media/ElectraCDM/Source/ElectraCDM/Private/ClearKey/ClearKeyCDM.cpp
+++ b/Engine/Plugins/Media/ElectraCDM/Source/ElectraCDM/Private/ClearKey/ClearKeyCDM.cpp
@@ -25,10 +25,18 @@ static FString UrlAttrib_clearkey_laurl(TEXT("clearkey:Laurl"));
 static FString UrlAttrib_Laurl(TEXT("Laurl"));
 };
 
+enum EClearKeyScheme
+{
+	Unsupported,
+	Cenc,
+	Cbcs
+};
+
 struct FClearKeyKIDKey
 {
 	TArray<uint8> KID;
 	TArray<uint8> Key;
+	EClearKeyScheme EncryptionScheme = EClearKeyScheme::Unsupported;
 };
 
 
@@ -44,12 +52,15 @@ public:
 	virtual ECDMError CreateDRMClient(TSharedPtr<IMediaCDMClient, ESPMode::ThreadSafe>& OutClient, IMediaCDM::IPlayerSession* InForPlayerSession, const TArray<IMediaCDM::FCDMCandidate>& InCandidates) override;
 	virtual ECDMError ReleasePlayerSessionKeys(IMediaCDM::IPlayerSession* PlayerSession) override;
 
+	virtual ESupportResult SupportsCipher(const FString& InCipherType) override;
 	virtual ESupportResult SupportsType(const FString& InMimeType) override;
 	virtual ESupportResult RequiresSecureDecoder(const FString& InMimeType) override;
 
 	void AddPlayerSessionKeys(IMediaCDM::IPlayerSession* InPlayerSession, const TArray<FClearKeyKIDKey>& InNewSessionKeys);
 	bool GetPlayerSessionKey(FClearKeyKIDKey& OutKey, IMediaCDM::IPlayerSession* InPlayerSession, const TArray<uint8>& InForKID);
 
+	static EClearKeyScheme GetCommonSchemeFromCipher(const FString& InCipherName);
+	static EClearKeyScheme GetCommonSchemeFromCipher(uint32 InCipher4CC);
 public:
 	static TSharedPtr<FClearKeyCDM, ESPMode::ThreadSafe> Get();
 
@@ -81,18 +92,19 @@ public:
 	virtual ECDMError BlockStreamDecryptEnd(IStreamDecryptHandle* InStreamDecryptContext) override;
 
 private:
-	struct FKeyDecrpyter
+	struct FKeyDecrypter
 	{
 		FClearKeyKIDKey KIDKey;
 		TSharedPtr<ElectraCDM::IStreamDecrypterAES128, ESPMode::ThreadSafe> Decrypter;
 		ECDMState State = ECDMState::Idle;
+		bool bIsInitialized = false;
 	};
 
-	TSharedPtr<ElectraCDM::IStreamDecrypterAES128, ESPMode::ThreadSafe> GetDecrypterForKID(const TArray<uint8>& KID);
+	FClearKeyDRMDecrypter::FKeyDecrypter*  GetDecrypterForKID(const TArray<uint8>& KID);
 
 	FCriticalSection Lock;
 	TArray<FClearKeyKIDKey> LicenseKeys;
-	TArray<FKeyDecrpyter> KeyDecrypters;
+	TArray<FKeyDecrypter> KeyDecrypters;
 	ECDMState CurrentState = ECDMState::Idle;
 	FString LastErrorMsg;
 };
@@ -119,6 +131,8 @@ public:
 	virtual ECDMError CreateDecrypter(TSharedPtr<IMediaCDMDecrypter, ESPMode::ThreadSafe>& OutDecrypter, const FString& InMimeType) override;
 
 private:
+	EClearKeyScheme GetCommonSchemeFromConfiguration(const IMediaCDM::FCDMCandidate& InConfiguration);
+	TArray<IMediaCDM::FCDMCandidate> GetConfigurationsForKID(const TArray<uint8>& InForKID);
 	int32 PrepareKIDsToRequest();
 	void AddKeyKID(const FClearKeyKIDKey& InKeyKid);
 	void AddKeyKIDs(const TArray<FClearKeyKIDKey>& InKeyKids);
@@ -203,9 +217,9 @@ namespace
 		FMemory::Memcpy(OutArray.GetData(), cnv.Get(), Len);
 	}
 
-	FString ArrayToString(const TArray<uint8>& InArray)
+	FString ArrayToString(const TArray<uint8>& InArray, int32 InStartAt=0)
 	{
-		FUTF8ToTCHAR cnv((const ANSICHAR*)InArray.GetData(), InArray.Num());
+		FUTF8ToTCHAR cnv((const ANSICHAR*)InArray.GetData() + InStartAt, InArray.Num() - InStartAt);
 		FString UTF8Text(cnv.Length(), cnv.Get());
 		return MoveTemp(UTF8Text);
 	}
@@ -360,6 +374,14 @@ ECDMError FClearKeyCDM::ReleasePlayerSessionKeys(IMediaCDM::IPlayerSession* Play
 }
 
 
+//-----------------------------------------------------------------------------
+/**
+ * Returns if a specified cipher (eg. "cenc" or "cbcs") is supported by this CDM.
+ */
+IMediaCDMCapabilities::ESupportResult FClearKeyCDM::SupportsCipher(const FString& InCipherType)
+{
+	return GetCommonSchemeFromCipher(InCipherType) != EClearKeyScheme::Unsupported ? IMediaCDMCapabilities::ESupportResult::Supported : IMediaCDMCapabilities::ESupportResult::NotSupported;
+}
 
 //-----------------------------------------------------------------------------
 /**
@@ -384,6 +406,37 @@ IMediaCDMCapabilities::ESupportResult FClearKeyCDM::RequiresSecureDecoder(const
 	return IMediaCDMCapabilities::ESupportResult::SecureDecoderNotRequired;
 }
 
+//-----------------------------------------------------------------------------
+/**
+ * Converts cipher name to enum.
+ */
+EClearKeyScheme FClearKeyCDM::GetCommonSchemeFromCipher(const FString& InCipherName)
+{
+	if (InCipherName.Equals(TEXT("cenc"), ESearchCase::IgnoreCase))
+	{
+		return EClearKeyScheme::Cenc;
+	}
+	else if (InCipherName.Equals(TEXT("cbcs"), ESearchCase::IgnoreCase))
+	{
+		return EClearKeyScheme::Cbcs;
+	}
+	return EClearKeyScheme::Unsupported;
+}
+
+EClearKeyScheme FClearKeyCDM::GetCommonSchemeFromCipher(uint32 InCipher4CC)
+{
+	#define MAKE_4CC(a,b,c,d) (uint32)(((uint32)a << 24) | ((uint32)b << 16) | ((uint32)c << 8) | ((uint32)d))
+	if (InCipher4CC == MAKE_4CC('c', 'e', 'n', 'c'))
+	{
+		return EClearKeyScheme::Cenc;
+	}
+	else if (InCipher4CC == MAKE_4CC('c', 'b', 'c', 's'))
+	{
+		return EClearKeyScheme::Cbcs;
+	}
+	return EClearKeyScheme::Unsupported;
+}
+
 
 /***************************************************************************************************************************************************/
 /***************************************************************************************************************************************************/
@@ -626,6 +679,43 @@ int32 FClearKeyDRMClient::PrepareKIDsToRequest()
 	return PendingRequiredKIDs.Num();
 }
 
+
+//-----------------------------------------------------------------------------
+/**
+ * Converts encryption scheme string to enum.
+ */
+EClearKeyScheme FClearKeyDRMClient::GetCommonSchemeFromConfiguration(const IMediaCDM::FCDMCandidate& InConfiguration)
+{
+	return FClearKeyCDM::GetCommonSchemeFromCipher(InConfiguration.CommonScheme);
+}
+
+
+//-----------------------------------------------------------------------------
+/**
+ * Returns the CDM configuration objects matching the given KID.
+ */
+TArray<IMediaCDM::FCDMCandidate> FClearKeyDRMClient::GetConfigurationsForKID(const TArray<uint8>& InForKID)
+{
+	TArray<IMediaCDM::FCDMCandidate> Cfgs;
+	FScopeLock lock(&Lock);
+	for(int32 nCfg=0; nCfg<CDMConfigurations.Num(); ++nCfg)
+	{
+		for(int32 nKIDs=0; nKIDs<CDMConfigurations[nCfg].DefaultKIDs.Num(); ++nKIDs)
+		{
+			if (CDMConfigurations[nCfg].DefaultKIDs[nKIDs].Len())
+			{
+				TArray<uint8> BinKID;
+				ConvertKIDToBin(BinKID, StripDashesFromKID(CDMConfigurations[nCfg].DefaultKIDs[nKIDs]));
+				if (BinKID == InForKID)
+				{
+					Cfgs.Emplace(CDMConfigurations[nCfg]);
+				}
+			}
+		}
+	}
+	return Cfgs;
+}
+
 //-----------------------------------------------------------------------------
 /**
  * Returns the information necessary to make the license request.
@@ -676,6 +766,8 @@ ECDMError FClearKeyDRMClient::SetLicenseKeyResponseData(void* InEventId, int32 H
 				FClearKeyKIDKey NewKey;
 				if (Base64UrlDecode(NewKey.KID, KID))
 				{
+					TArray<IMediaCDM::FCDMCandidate> Configs(GetConfigurationsForKID(NewKey.KID));
+					NewKey.EncryptionScheme = Configs.Num() ? GetCommonSchemeFromConfiguration(Configs[0]) : EClearKeyScheme::Unsupported;
 					NewKey.Key = InKeyResponseData;
 					NewLicenseKeys.Emplace(MoveTemp(NewKey));
 				}
@@ -684,7 +776,9 @@ ECDMError FClearKeyDRMClient::SetLicenseKeyResponseData(void* InEventId, int32 H
 		else
 #endif
 		{
-			FString JsonString = ArrayToString(InKeyResponseData);
+			// Is the JSON starting with an UTF-8 BOM?
+			int32 ResponseOffset = InKeyResponseData.Num() > 3 && InKeyResponseData[0] == 0xEF && InKeyResponseData[1] == 0xBB && InKeyResponseData[2] == 0xBF ? 3 : 0;
+			FString JsonString = ArrayToString(InKeyResponseData, ResponseOffset);
 			TSharedRef<TJsonReader<>> Reader = TJsonReaderFactory<>::Create(JsonString);
 			TSharedPtr<FJsonObject> KeyResponse;
 			if (FJsonSerializer::Deserialize(Reader, KeyResponse))
@@ -709,6 +803,11 @@ ECDMError FClearKeyDRMClient::SetLicenseKeyResponseData(void* InEventId, int32 H
 							}
 							if (bOk)
 							{
+								TArray<IMediaCDM::FCDMCandidate> Configs(GetConfigurationsForKID(NewKey.KID));
+								// We require all KIDs to be unambiguously using the same encryption scheme
+								// so it is sufficient to look at the first configuration.
+								// If same KID configurations use different schemes this is an authoring error.
+								NewKey.EncryptionScheme = Configs.Num() ? GetCommonSchemeFromConfiguration(Configs[0]) : EClearKeyScheme::Unsupported;
 								NewLicenseKeys.Emplace(MoveTemp(NewKey));
 							}
 							else
@@ -1088,7 +1187,7 @@ ECDMError FClearKeyDRMDecrypter::UpdateFromURL(const FString& InURL, const FStri
 /**
  * Locates the decrypter for the given key ID.
  */
-TSharedPtr<ElectraCDM::IStreamDecrypterAES128, ESPMode::ThreadSafe> FClearKeyDRMDecrypter::GetDecrypterForKID(const TArray<uint8>& KID)
+FClearKeyDRMDecrypter::FKeyDecrypter* FClearKeyDRMDecrypter::GetDecrypterForKID(const TArray<uint8>& KID)
 {
 	// Note: The critical section must be locked already!
 	for(int32 i=0; i<KeyDecrypters.Num(); ++i)
@@ -1098,7 +1197,7 @@ TSharedPtr<ElectraCDM::IStreamDecrypterAES128, ESPMode::ThreadSafe> FClearKeyDRM
 			// Decrypter needs to be ready.
 			if (KeyDecrypters[i].State == ECDMState::Ready)
 			{
-				return KeyDecrypters[i].Decrypter;
+				return &KeyDecrypters[i];
 			}
 		}
 	}
@@ -1116,24 +1215,14 @@ void FClearKeyDRMDecrypter::Reinitialize()
 	KeyDecrypters.Empty();
 	LastErrorMsg.Empty();
 
-	ElectraCDM::IStreamDecrypterAES128::EResult Result;
 	CurrentState = ECDMState::Ready;
 	for(int32 i=0; i<LicenseKeys.Num(); ++i)
 	{
-		FKeyDecrpyter& kd = KeyDecrypters.AddDefaulted_GetRef();
+		FKeyDecrypter& kd = KeyDecrypters.AddDefaulted_GetRef();
 		kd.KIDKey = LicenseKeys[i];
 		kd.Decrypter = ElectraCDM::IStreamDecrypterAES128::Create();
-		Result = kd.Decrypter->CTRInit(kd.KIDKey.Key);
-		if (Result == ElectraCDM::IStreamDecrypterAES128::EResult::Ok)
-		{
-			kd.State = ECDMState::Ready;
-		}
-		else
-		{
-			kd.State = ECDMState::InvalidKey;
-			CurrentState = ECDMState::InvalidKey;
-			LastErrorMsg = TEXT("Invalid key");
-		}
+		kd.State = ECDMState::Ready;
+		kd.bIsInitialized = false;
 	}
 }
 
@@ -1145,31 +1234,142 @@ ECDMError FClearKeyDRMDecrypter::DecryptInPlace(uint8* InOutData, int32 InNumDat
 {
 	FScopeLock lock(&Lock);
 	LastErrorMsg.Empty();
-	TSharedPtr<ElectraCDM::IStreamDecrypterAES128, ESPMode::ThreadSafe> Decrypter = GetDecrypterForKID(InSampleInfo.DefaultKID);
+	FClearKeyDRMDecrypter::FKeyDecrypter* DecrypterState = GetDecrypterForKID(InSampleInfo.DefaultKID);
+	TSharedPtr<ElectraCDM::IStreamDecrypterAES128, ESPMode::ThreadSafe> Decrypter = DecrypterState ? DecrypterState->Decrypter : nullptr;
 	if (Decrypter.IsValid())
 	{
-		if (Decrypter->CTRSetIV(InSampleInfo.IV) != ElectraCDM::IStreamDecrypterAES128::EResult::Ok)
-		{
-			LastErrorMsg = TEXT("Bad IV");
-			return ECDMError::Failure;
-		}
-		if (InSampleInfo.SubSamples.Num() == 0)
+		ElectraCDM::IStreamDecrypterAES128::EResult Result;
+
+		EClearKeyScheme SchemeFromMedia = FClearKeyCDM::GetCommonSchemeFromCipher(InSampleInfo.Scheme4CC);
+		EClearKeyScheme SchemeFromKID = DecrypterState->KIDKey.EncryptionScheme;
+
+		EClearKeyScheme SchemeToUse = SchemeFromMedia != EClearKeyScheme::Unsupported ? SchemeFromMedia : SchemeFromKID;
+
+		// "cenc" scheme? (AES-128 CTR)
+		if (SchemeToUse == EClearKeyScheme::Cenc)
 		{
-			Decrypter->CTRDecryptInPlace(InOutData, InNumDataBytes);
-			return ECDMError::Success;
+			if (!DecrypterState->bIsInitialized)
+			{
+				Result = Decrypter->CTRInit(DecrypterState->KIDKey.Key);
+				if ((DecrypterState->bIsInitialized = Result == ElectraCDM::IStreamDecrypterAES128::EResult::Ok) == false)
+				{
+					DecrypterState->State = ECDMState::InvalidKey;
+					CurrentState = ECDMState::InvalidKey;
+					LastErrorMsg = TEXT("Invalid key");
+					return ECDMError::Failure;
+				}
+			}
+
+			if (Decrypter->CTRSetIV(InSampleInfo.IV) != ElectraCDM::IStreamDecrypterAES128::EResult::Ok)
+			{
+				LastErrorMsg = TEXT("Bad IV");
+				return ECDMError::Failure;
+			}
+			if (InSampleInfo.SubSamples.Num() == 0)
+			{
+				Decrypter->CTRDecryptInPlace(InOutData, InNumDataBytes);
+				return ECDMError::Success;
+			}
+			else
+			{
+				for(int32 i=0; i<InSampleInfo.SubSamples.Num(); ++i)
+				{
+					InOutData += InSampleInfo.SubSamples[i].NumClearBytes;
+					if (InSampleInfo.SubSamples[i].NumEncryptedBytes)
+					{
+						Decrypter->CTRDecryptInPlace(InOutData, InSampleInfo.SubSamples[i].NumEncryptedBytes);
+					}
+					InOutData += InSampleInfo.SubSamples[i].NumEncryptedBytes;
+				}
+				return ECDMError::Success;
+			}
 		}
-		else
+		// "cbcs" scheme? (AES-128 CBC)
+		else if (SchemeToUse == EClearKeyScheme::Cbcs)
 		{
-			for(int32 i=0; i<InSampleInfo.SubSamples.Num(); ++i)
+			int32 NumBytesDecrypted = 0;
+
+			auto DecryptPatternBlock = [&](int32 BlocksToGo) -> void
+			{
+				while(BlocksToGo > 0)
+				{
+					int32 NumEnc = BlocksToGo >= InSampleInfo.Pattern.CryptByteBlock ? InSampleInfo.Pattern.CryptByteBlock : BlocksToGo;
+					/*Result =*/ Decrypter->CBCDecryptInPlace(NumBytesDecrypted, InOutData, NumEnc * 16, false);
+					InOutData += (NumEnc + InSampleInfo.Pattern.SkipByteBlock) * 16;
+					BlocksToGo -= InSampleInfo.Pattern.CryptByteBlock + InSampleInfo.Pattern.SkipByteBlock;
+				}
+			};
+
+			if (InSampleInfo.SubSamples.Num() == 0)
+			{
+				Result = DecrypterState->Decrypter->CBCInit(DecrypterState->KIDKey.Key, &InSampleInfo.IV);
+				if (Result != ElectraCDM::IStreamDecrypterAES128::EResult::Ok)
+				{
+					LastErrorMsg = TEXT("Bad key or IV length");
+					return ECDMError::Failure;
+				}
+
+				// Entire sample is encrypted.
+				if (InSampleInfo.Pattern.CryptByteBlock == 0 && InSampleInfo.Pattern.SkipByteBlock == 0)
+				{
+					/*Result =*/ Decrypter->CBCDecryptInPlace(NumBytesDecrypted, InOutData, InNumDataBytes & ~15, false);
+				}
+				else
+				{
+					DecryptPatternBlock(InNumDataBytes / 16);
+				}
+				return ECDMError::Success;
+			}
+			else
 			{
-				InOutData += InSampleInfo.SubSamples[i].NumClearBytes;
-				if (InSampleInfo.SubSamples[i].NumEncryptedBytes)
+				for(int32 i=0; i<InSampleInfo.SubSamples.Num(); ++i)
 				{
-					Decrypter->CTRDecryptInPlace(InOutData, InSampleInfo.SubSamples[i].NumEncryptedBytes);
+					// cbcs encryption is restarted with every sub-sample.
+					if (InSampleInfo.IV.Num() == 16)
+					{
+						Result = DecrypterState->Decrypter->CBCInit(DecrypterState->KIDKey.Key, &InSampleInfo.IV);
+					}
+					else if (InSampleInfo.IV.Num() < 16)
+					{
+						TArray<uint8> paddedIV(InSampleInfo.IV);
+						paddedIV.InsertZeroed(paddedIV.Num(), 16-paddedIV.Num());
+						Result = DecrypterState->Decrypter->CBCInit(DecrypterState->KIDKey.Key, &paddedIV);
+					}
+					else
+					{
+						Result = ElectraCDM::IStreamDecrypterAES128::EResult::BadIVLength;
+					}
+					if (Result != ElectraCDM::IStreamDecrypterAES128::EResult::Ok)
+					{
+						LastErrorMsg = TEXT("Bad key or IV length");
+						return ECDMError::Failure;
+					}
+
+					InOutData += InSampleInfo.SubSamples[i].NumClearBytes;
+
+					if (InSampleInfo.Pattern.CryptByteBlock == 0 && InSampleInfo.Pattern.SkipByteBlock == 0)
+					{
+						/*Result =*/ Decrypter->CBCDecryptInPlace(NumBytesDecrypted, InOutData, InSampleInfo.SubSamples[i].NumEncryptedBytes & ~15, false);
+						InOutData = InOutData + InSampleInfo.SubSamples[i].NumEncryptedBytes;
+					}
+					else
+					{
+						uint8* EncDataStart = InOutData;
+						DecryptPatternBlock((int32) InSampleInfo.SubSamples[i].NumEncryptedBytes / 16);
+						// The number of encrypted bytes in the subsample is not necessarily a 16B multiple.
+						// If not, the last few bytes are not encrypted and as such are not touched by
+						// decryption. We need to advance the pointer to the data to the actual number of
+						// bytes - NOT the number of bytes that were decrypted! - to get to the next subsample.
+						InOutData = EncDataStart + InSampleInfo.SubSamples[i].NumEncryptedBytes;
+					}
 				}
-				InOutData += InSampleInfo.SubSamples[i].NumEncryptedBytes;
+				return ECDMError::Success;
 			}
-			return ECDMError::Success;
+		}
+		else
+		{
+			LastErrorMsg = TEXT("Unsupported encryption scheme for KID");
+			return ECDMError::Failure;
 		}
 	}
 	LastErrorMsg = TEXT("No valid decrypter found for KID");
diff --git a/Engine/Plugins/Media/ElectraCDM/Source/ElectraCDM/Private/Crypto/StreamCryptoAES128.cpp b/Engine/Plugins/Media/ElectraCDM/Source/ElectraCDM/Private/Crypto/StreamCryptoAES128.cpp
--- a/Engine/Plugins/Media/ElectraCDM/Source/ElectraCDM/Private/Crypto/StreamCryptoAES128.cpp
+++ b/Engine/Plugins/Media/ElectraCDM/Source/ElectraCDM/Private/Crypto/StreamCryptoAES128.cpp
@@ -445,7 +445,8 @@ IStreamDecrypterAES128::EResult FStreamDecrypterAES128::CBCInit(const TArray<uin
 		}
 		else
 		{
-			TinyAES128::AES_init_ctx(&AES128CBCContext, (const uint8_t*)Key.GetData());
+			uint8 ZeroIV[16] = {0};
+			TinyAES128::AES_init_ctx_iv(&AES128CBCContext, (const uint8_t*)Key.GetData(), ZeroIV);
 		}
 		bIsInitialized = true;
 		CryptoKey = Key;
diff --git a/Engine/Plugins/Media/ElectraCDM/Source/ElectraCDM/Public/ElectraCDM.h b/Engine/Plugins/Media/ElectraCDM/Source/ElectraCDM/Public/ElectraCDM.h
--- a/Engine/Plugins/Media/ElectraCDM/Source/ElectraCDM/Public/ElectraCDM.h
+++ b/Engine/Plugins/Media/ElectraCDM/Source/ElectraCDM/Public/ElectraCDM.h
@@ -55,6 +55,9 @@ public:
 		SecureDecoderNotRequired,
 	};
 
+	// Returns if a specified cipher (eg. "cenc" or "cbcs") is supported by this CDM.
+	virtual ESupportResult SupportsCipher(const FString& InCipherType) = 0;
+
 	// Returns if a specified media type is supported by this CDM.
 	virtual ESupportResult SupportsType(const FString& InMimeType) = 0;
 
diff --git a/Engine/Plugins/Media/ElectraCDM/Source/ElectraCDM/Public/ElectraEncryptedSampleInfo.h b/Engine/Plugins/Media/ElectraCDM/Source/ElectraCDM/Public/ElectraEncryptedSampleInfo.h
--- a/Engine/Plugins/Media/ElectraCDM/Source/ElectraCDM/Public/ElectraEncryptedSampleInfo.h
+++ b/Engine/Plugins/Media/ElectraCDM/Source/ElectraCDM/Public/ElectraEncryptedSampleInfo.h
@@ -27,6 +27,7 @@ struct FMediaCDMSampleInfo
 	TArray<uint8> IV;
 	TArray<FSubSample> SubSamples;
 	FCryptSkipPattern Pattern;
+	uint32 Scheme4CC = 0;
 };
 
 
diff --git a/Engine/Plugins/Media/ElectraPlayer/Source/ElectraPlayerRuntime/Private/Runtime/Demuxer/ParserISO14496-12.cpp b/Engine/Plugins/Media/ElectraPlayer/Source/ElectraPlayerRuntime/Private/Runtime/Demuxer/ParserISO14496-12.cpp
--- a/Engine/Plugins/Media/ElectraPlayer/Source/ElectraPlayerRuntime/Private/Runtime/Demuxer/ParserISO14496-12.cpp
+++ b/Engine/Plugins/Media/ElectraPlayer/Source/ElectraPlayerRuntime/Private/Runtime/Demuxer/ParserISO14496-12.cpp
@@ -765,6 +765,7 @@ private:
 		 * Encryption scheme types
 		 */
 		static const IParserISO14496_12::FBoxType kEncryptionScheme_cenc = MAKE_BOX_ATOM('c', 'e', 'n', 'c');
+		static const IParserISO14496_12::FBoxType kEncryptionScheme_cbcs = MAKE_BOX_ATOM('c', 'b', 'c', 's');
 
 		/**
 		 * Known UUIDs for PIFF boxes before they were added to ISO/IEC 14496-12
@@ -6534,23 +6535,36 @@ private:
 
 	bool FParserISO14496_12::FTrackIterator::GetEncryptionInfo(ElectraCDM::FMediaCDMSampleInfo& OutSampleEncryptionInfo) const
 	{
-		if (!Track || !Track->SENCBox)
+		if (!Track || (!Track->SENCBox && !Track->TENCBox))
 		{
 			return false;
 		}
 		if (Track->TENCBox)
 		{
 			OutSampleEncryptionInfo.DefaultKID = Track->TENCBox->GetDefaultKID();
+			OutSampleEncryptionInfo.IV = Track->TENCBox->GetDefaultConstantIV();
+			if (Track->TENCBox->HasDefaultCryptBlockValues())
+			{
+				OutSampleEncryptionInfo.Pattern.PatternType = 1;
+				OutSampleEncryptionInfo.Pattern.CryptByteBlock = Track->TENCBox->GetDefaultCryptByteBlock();
+				OutSampleEncryptionInfo.Pattern.SkipByteBlock = Track->TENCBox->GetDefaultSkipByteBlock();
+			}
 		}
-
-		const TArray<FMP4BoxSENC::FEntry>& SencEntries = Track->SENCBox->GetEntries();
-		check((int32) SampleNumberInTRUN < SencEntries.Num());
-		if ((int32) SampleNumberInTRUN >= SencEntries.Num())
+		if (Track->SENCBox)
 		{
-			return false;
+			const TArray<FMP4BoxSENC::FEntry>& SencEntries = Track->SENCBox->GetEntries();
+			check((int32) SampleNumberInTRUN < SencEntries.Num());
+			if ((int32) SampleNumberInTRUN >= SencEntries.Num())
+			{
+				return false;
+			}
+			if (SencEntries[SampleNumberInTRUN].IV.Num())
+			{
+				OutSampleEncryptionInfo.IV = SencEntries[SampleNumberInTRUN].IV;
+			}
+			OutSampleEncryptionInfo.SubSamples = SencEntries[SampleNumberInTRUN].SubSamples;
 		}
-		OutSampleEncryptionInfo.IV = SencEntries[SampleNumberInTRUN].IV;
-		OutSampleEncryptionInfo.SubSamples = SencEntries[SampleNumberInTRUN].SubSamples;
+		OutSampleEncryptionInfo.Scheme4CC = Track->SCHMBox ? Track->SCHMBox->GetSchemeType() : 0;
 		return true;
 	}
 
@@ -7148,8 +7162,8 @@ private:
 								{
 									return UEMEDIA_ERROR_FORMAT_ERROR;
 								}
-								// Expecting encryption to be 'cenc'
-								if (SCHMBox->GetSchemeType() != FMP4Box::kEncryptionScheme_cenc)
+								// Expecting encryption to be 'cenc' or 'cbcs'
+								if (SCHMBox->GetSchemeType() != FMP4Box::kEncryptionScheme_cenc && SCHMBox->GetSchemeType() != FMP4Box::kEncryptionScheme_cbcs)
 								{
 									return UEMEDIA_ERROR_FORMAT_ERROR;
 								}
@@ -7197,8 +7211,8 @@ private:
 								{
 									return UEMEDIA_ERROR_FORMAT_ERROR;
 								}
-								// Likewise fo encryption it needs to be 'cenc'
-								if (SCHMBox->GetSchemeType() != FMP4Box::kEncryptionScheme_cenc)
+								// Likewise fo encryption it needs to be 'cenc' or 'cbcs'
+								if (SCHMBox->GetSchemeType() != FMP4Box::kEncryptionScheme_cenc && SCHMBox->GetSchemeType() != FMP4Box::kEncryptionScheme_cbcs)
 								{
 									return UEMEDIA_ERROR_FORMAT_ERROR;
 								}
