ElectraCDM: fixed a static analysis warning

#rnx

diff --git a/Engine/Plugins/Media/ElectraCDM/Source/ElectraCDM/Private/ClearKey/ClearKeyCDM.cpp b/Engine/Plugins/Media/ElectraCDM/Source/ElectraCDM/Private/ClearKey/ClearKeyCDM.cpp
--- a/Engine/Plugins/Media/ElectraCDM/Source/ElectraCDM/Private/ClearKey/ClearKeyCDM.cpp
+++ b/Engine/Plugins/Media/ElectraCDM/Source/ElectraCDM/Private/ClearKey/ClearKeyCDM.cpp
@@ -1236,7 +1236,7 @@ ECDMError FClearKeyDRMDecrypter::DecryptInPlace(uint8* InOutData, int32 InNumDat
 	LastErrorMsg.Empty();
 	FClearKeyDRMDecrypter::FKeyDecrypter* DecrypterState = GetDecrypterForKID(InSampleInfo.DefaultKID);
 	TSharedPtr<ElectraCDM::IStreamDecrypterAES128, ESPMode::ThreadSafe> Decrypter = DecrypterState ? DecrypterState->Decrypter : nullptr;
-	if (Decrypter.IsValid())
+	if (Decrypter.IsValid() && DecrypterState)
 	{
 		ElectraCDM::IStreamDecrypterAES128::EResult Result;
 
