ElectraCDM: initializing the cbc decrypter only if the sub sample has encrypted bytes

#rnx

diff --git a/Engine/Plugins/Media/ElectraCDM/Source/ElectraCDM/Private/ClearKey/ClearKeyCDM.cpp b/Engine/Plugins/Media/ElectraCDM/Source/ElectraCDM/Private/ClearKey/ClearKeyCDM.cpp
--- a/Engine/Plugins/Media/ElectraCDM/Source/ElectraCDM/Private/ClearKey/ClearKeyCDM.cpp
+++ b/Engine/Plugins/Media/ElectraCDM/Source/ElectraCDM/Private/ClearKey/ClearKeyCDM.cpp
@@ -1324,43 +1324,46 @@ ECDMError FClearKeyDRMDecrypter::DecryptInPlace(uint8* InOutData, int32 InNumDat
 			{
 				for(int32 i=0; i<InSampleInfo.SubSamples.Num(); ++i)
 				{
-					// cbcs encryption is restarted with every sub-sample.
-					if (InSampleInfo.IV.Num() == 16)
-					{
-						Result = DecrypterState->Decrypter->CBCInit(DecrypterState->KIDKey.Key, &InSampleInfo.IV);
-					}
-					else if (InSampleInfo.IV.Num() < 16)
-					{
-						TArray<uint8> paddedIV(InSampleInfo.IV);
-						paddedIV.InsertZeroed(paddedIV.Num(), 16-paddedIV.Num());
-						Result = DecrypterState->Decrypter->CBCInit(DecrypterState->KIDKey.Key, &paddedIV);
-					}
-					else
-					{
-						Result = ElectraCDM::IStreamDecrypterAES128::EResult::BadIVLength;
-					}
-					if (Result != ElectraCDM::IStreamDecrypterAES128::EResult::Ok)
-					{
-						LastErrorMsg = TEXT("Bad key or IV length");
-						return ECDMError::Failure;
-					}
-
 					InOutData += InSampleInfo.SubSamples[i].NumClearBytes;
 
-					if (InSampleInfo.Pattern.CryptByteBlock == 0 && InSampleInfo.Pattern.SkipByteBlock == 0)
-					{
-						/*Result =*/ Decrypter->CBCDecryptInPlace(NumBytesDecrypted, InOutData, InSampleInfo.SubSamples[i].NumEncryptedBytes & ~15, false);
-						InOutData = InOutData + InSampleInfo.SubSamples[i].NumEncryptedBytes;
-					}
-					else
+					if (InSampleInfo.SubSamples[i].NumEncryptedBytes)
 					{
-						uint8* EncDataStart = InOutData;
-						DecryptPatternBlock((int32) InSampleInfo.SubSamples[i].NumEncryptedBytes / 16);
-						// The number of encrypted bytes in the subsample is not necessarily a 16B multiple.
-						// If not, the last few bytes are not encrypted and as such are not touched by
-						// decryption. We need to advance the pointer to the data to the actual number of
-						// bytes - NOT the number of bytes that were decrypted! - to get to the next subsample.
-						InOutData = EncDataStart + InSampleInfo.SubSamples[i].NumEncryptedBytes;
+						// cbcs encryption is restarted with every sub-sample.
+						if (InSampleInfo.IV.Num() == 16)
+						{
+							Result = DecrypterState->Decrypter->CBCInit(DecrypterState->KIDKey.Key, &InSampleInfo.IV);
+						}
+						else if (InSampleInfo.IV.Num() < 16)
+						{
+							TArray<uint8> paddedIV(InSampleInfo.IV);
+							paddedIV.InsertZeroed(paddedIV.Num(), 16-paddedIV.Num());
+							Result = DecrypterState->Decrypter->CBCInit(DecrypterState->KIDKey.Key, &paddedIV);
+						}
+						else
+						{
+							Result = ElectraCDM::IStreamDecrypterAES128::EResult::BadIVLength;
+						}
+						if (Result != ElectraCDM::IStreamDecrypterAES128::EResult::Ok)
+						{
+							LastErrorMsg = TEXT("Bad key or IV length");
+							return ECDMError::Failure;
+						}
+
+						if (InSampleInfo.Pattern.CryptByteBlock == 0 && InSampleInfo.Pattern.SkipByteBlock == 0)
+						{
+							/*Result =*/ Decrypter->CBCDecryptInPlace(NumBytesDecrypted, InOutData, InSampleInfo.SubSamples[i].NumEncryptedBytes & ~15, false);
+							InOutData = InOutData + InSampleInfo.SubSamples[i].NumEncryptedBytes;
+						}
+						else
+						{
+							uint8* EncDataStart = InOutData;
+							DecryptPatternBlock((int32) InSampleInfo.SubSamples[i].NumEncryptedBytes / 16);
+							// The number of encrypted bytes in the subsample is not necessarily a 16B multiple.
+							// If not, the last few bytes are not encrypted and as such are not touched by
+							// decryption. We need to advance the pointer to the data to the actual number of
+							// bytes - NOT the number of bytes that were decrypted! - to get to the next subsample.
+							InOutData = EncDataStart + InSampleInfo.SubSamples[i].NumEncryptedBytes;
+						}
 					}
 				}
 				return ECDMError::Success;
